Integrating CL# 9474758 from Dev-Enterprise

Fixing crosscompile linux build

diff --git a/Engine/Plugins/Enterprise/StaticMeshEditorExtension/StaticMeshEditorExtension.uplugin b/Engine/Plugins/Enterprise/StaticMeshEditorExtension/StaticMeshEditorExtension.uplugin
--- a/Engine/Plugins/Enterprise/StaticMeshEditorExtension/StaticMeshEditorExtension.uplugin
+++ b/Engine/Plugins/Enterprise/StaticMeshEditorExtension/StaticMeshEditorExtension.uplugin
@@ -19,14 +19,12 @@
 		{
 			"Name" : "StaticMeshEditorExtension",
 			"Type" : "Editor",
-			"LoadingPhase" : "Default",
-			"WhitelistPlatforms" : [ "Win64" ]
+			"LoadingPhase" : "Default"
 		},
 		{
 			"Name" : "MeshProcessingLibrary",
 			"Type" : "Editor",
-			"LoadingPhase" : "Default",
-			"WhitelistPlatforms": [ "Win64" ]
+			"LoadingPhase" : "Default"
 		},
 	],
 	"Plugins": [
