Move the Unwrap UV command from SME toolbar to UV rollout

diff --git a/Engine/Plugins/Enterprise/StaticMeshEditorExtension/Source/StaticMeshEditorExtension/Private/StaticMeshEditorExtensionModule.cpp b/Engine/Plugins/Enterprise/StaticMeshEditorExtension/Source/StaticMeshEditorExtension/Private/StaticMeshEditorExtensionModule.cpp
--- a/Engine/Plugins/Enterprise/StaticMeshEditorExtension/Source/StaticMeshEditorExtension/Private/StaticMeshEditorExtensionModule.cpp
+++ b/Engine/Plugins/Enterprise/StaticMeshEditorExtension/Source/StaticMeshEditorExtension/Private/StaticMeshEditorExtensionModule.cpp
@@ -170,6 +170,13 @@ TSharedRef<FExtender> FStaticMeshEditorExtensionModule::ExtendStaticMeshEditorTo
 		FMenuExtensionDelegate::CreateStatic(&FUVGenerationToolbar::CreateUVMenu, CommandList, Cast<UStaticMesh>(Objects[0]))
 	);
 
+	Extender->AddMenuExtension(
+		"UVActionOptions",
+		EExtensionHook::First,
+		CommandList,
+		FMenuExtensionDelegate::CreateStatic(&FUVGenerationFlattenMappingToolbar::CreateMenu, CommandList, Cast<UStaticMesh>(Objects[0]))
+	);
+
 	return Extender;
 }
 
@@ -188,13 +195,6 @@ TSharedRef<FExtender> FStaticMeshEditorExtensionModule::ExtendStaticMeshEditorSe
 		FToolBarExtensionDelegate::CreateStatic(&FPolygonEditingToolbar::CreateToolbar, CommandList, Cast<UStaticMesh>(Objects[0]) )
 	);
 
-	Extender->AddToolBarExtension(
-		"Extensions",
-		EExtensionHook::After,
-		CommandList,
-		FToolBarExtensionDelegate::CreateStatic(&FUVGenerationFlattenMappingToolbar::CreateToolbar, CommandList, Cast<UStaticMesh>(Objects[0]))
-	);
-
 	return Extender;
 }
 
diff --git a/Engine/Plugins/Enterprise/StaticMeshEditorExtension/Source/StaticMeshEditorExtension/Private/UVTools/UVGenerationFlattenMappingTool.cpp b/Engine/Plugins/Enterprise/StaticMeshEditorExtension/Source/StaticMeshEditorExtension/Private/UVTools/UVGenerationFlattenMappingTool.cpp
--- a/Engine/Plugins/Enterprise/StaticMeshEditorExtension/Source/StaticMeshEditorExtension/Private/UVTools/UVGenerationFlattenMappingTool.cpp
+++ b/Engine/Plugins/Enterprise/StaticMeshEditorExtension/Source/StaticMeshEditorExtension/Private/UVTools/UVGenerationFlattenMappingTool.cpp
@@ -183,7 +183,7 @@ TSharedRef<FExtender> FUVGenerationFlattenMappingTool::OnExtendContentBrowserAss
 			MenuBuilder.AddMenuEntry(
 				LOCTEXT("ObjectContext_UnwrapUV", "Unwrap UV"),
 				LOCTEXT("ObjectContext_UnwrapUVTooltip", "Opens Unwrap UV option window"),
-				FSlateIcon(FUVGenerationFlattenMappingToolStyle::GetStyleSetName(), "UVGenerationFlattenMapping.UnwrapUV.Small"),
+				FSlateIcon(),
 				FUIAction(FExecuteAction::CreateStatic( &FUVGenerationFlattenMappingTool::OpenUnwrapUVWindow, StaticMeshes), FCanExecuteAction()));
 		}));
 	}
@@ -324,97 +324,25 @@ void SUVGenerationFlattenMappingWindow::Construct(const FArguments& InArgs)
 	DetailsViewBox->SetContent(DetailsView.ToSharedRef());
 }
 
-FUVGenerationFlattenMappingCommands::FUVGenerationFlattenMappingCommands()
-	: TCommands<FUVGenerationFlattenMappingCommands>(
-		"UVGenerationFlattenMapping",
-		LOCTEXT("FUVGenerationFlattenMappingCommands", "Static Mesh UV Unwrapping Commands"),
-		"StaticMeshEditor",
-		FUVGenerationFlattenMappingToolStyle::GetStyleSetName())
+void FUVGenerationFlattenMappingToolbar::CreateMenu(FMenuBuilder& ParentMenuBuilder, const TSharedRef<FUICommandList> CommandList, UStaticMesh* InStaticMesh)
 {
-}
-
-void FUVGenerationFlattenMappingCommands::RegisterCommands()
-{
-	UI_COMMAND(UnwrapUV, "Unwrap UV", "Opens Unwrap UV option window", EUserInterfaceActionType::Button, FInputChord(EKeys::U, false, true, false, false));
-}
-
-FUVGenerationFlattenMappingToolbar::FUVGenerationFlattenMappingToolbar()
-	: StaticMesh(nullptr)
-	, StaticMeshEditor(nullptr)
-	, BoundCommandList(nullptr)
-{
-	UVGenerationFlattenMappingToolbarProxyObject = TStrongObjectPtr<UUVGenerationFlattenMappingToolbarProxyObject>(NewObject<UUVGenerationFlattenMappingToolbarProxyObject>());
-	UVGenerationFlattenMappingToolbarProxyObject->Owner = this;
-}
-
-FUVGenerationFlattenMappingToolbar::~FUVGenerationFlattenMappingToolbar()
-{
-	UVGenerationFlattenMappingToolbarProxyObject->Owner = nullptr;
-	BoundCommandList = nullptr;
-	StaticMesh = nullptr;
-}
-
-void FUVGenerationFlattenMappingToolbar::CreateToolbar(FToolBarBuilder& ToolbarBuilder, const TSharedRef<FUICommandList> CommandList, UStaticMesh* InStaticMesh)
-{
-	TSharedPtr<FUVGenerationFlattenMappingToolbar> UVGenerationFlattenMappingToolbar = MakeShareable(new FUVGenerationFlattenMappingToolbar());
-
-	if (!UVGenerationFlattenMappingToolbar->Initialize(InStaticMesh, CommandList))
-	{
-		return;
-	}
-
-	UVGenerationFlattenMappingToolbar->PopulateToolbar(ToolbarBuilder, CommandList);
-}
-
-bool FUVGenerationFlattenMappingToolbar::Initialize(UStaticMesh* InStaticMesh, const TSharedRef<FUICommandList> CommandList)
-{
-	// Take a hold on the StaticMesh Editor hosting this toolbar
-	IAssetEditorInstance* EditorInstance = GEditor->GetEditorSubsystem<UAssetEditorSubsystem>()->FindEditorForAsset(InStaticMesh, false);
-	if (!EditorInstance || !EditorInstance->GetEditorName().ToString().Contains(TEXT("StaticMeshEditor")))
+	FUIAction GenerateUnwrappedUVMenuAction;
+	GenerateUnwrappedUVMenuAction.ExecuteAction.BindLambda([InStaticMesh]()
 	{
-		return false;
-	}
-
-	StaticMesh = InStaticMesh;
-
-	StaticMeshEditor = static_cast<IStaticMeshEditor*>(EditorInstance);
-
-	BindCommands(CommandList);
-
-	return true;
-}
-
-void FUVGenerationFlattenMappingToolbar::BindCommands(const TSharedPtr<FUICommandList> CommandList)
-{
-	BoundCommandList = CommandList;
-
-	// Initialize style set associated with MeshEditor plugin
-	FUVGenerationFlattenMappingToolStyle::Initialize();
-
-	FUVGenerationFlattenMappingCommands::Register();
-
-	const FUVGenerationFlattenMappingCommands& PolygonEditingCommands = FUVGenerationFlattenMappingCommands::Get();
-
-	CommandList->MapAction(
-		PolygonEditingCommands.UnwrapUV,
-		FExecuteAction::CreateLambda([this](){
-			FUVGenerationFlattenMappingTool::OpenUnwrapUVWindow({ StaticMesh });
-			StaticMeshEditor->RefreshTool();
-		}),
-		FCanExecuteAction()
+		IAssetEditorInstance* EditorInstance = GEditor->GetEditorSubsystem<UAssetEditorSubsystem>()->FindEditorForAsset(InStaticMesh, false);
+		if (!EditorInstance || !EditorInstance->GetEditorName().ToString().Contains(TEXT("StaticMeshEditor")))
+		{
+			return;
+		}
+		FUVGenerationFlattenMappingTool::OpenUnwrapUVWindow({ InStaticMesh });
+		static_cast<IStaticMeshEditor*>(EditorInstance)->RefreshTool();
+	});
+	ParentMenuBuilder.AddMenuEntry(
+		LOCTEXT("UnwrapUV", "Unwrap UV"),
+		LOCTEXT("UnwrapUVTooltip", "Opens the UV unwrapping window"),
+		FSlateIcon(),
+		GenerateUnwrappedUVMenuAction
 	);
 }
 
-void FUVGenerationFlattenMappingToolbar::PopulateToolbar(FToolBarBuilder& ToolbarBuilder, const TSharedRef<FUICommandList> CommandList)
-{
-	ToolbarBuilder.BeginSection("UVGeneration");
-	{
-		// Add invisible widget used to detect closure of hosting static mesh editor
-		ToolbarBuilder.AddWidget(SNew(SToolbarWidget<FUVGenerationFlattenMappingToolbar>).EditingToolbar(SharedThis(this)), NAME_None);
-
-		ToolbarBuilder.AddToolBarButton(FUVGenerationFlattenMappingCommands::Get().UnwrapUV);
-	}
-	ToolbarBuilder.EndSection();
-}
-
 #undef LOCTEXT_NAMESPACE
diff --git a/Engine/Plugins/Enterprise/StaticMeshEditorExtension/Source/StaticMeshEditorExtension/Private/UVTools/UVGenerationFlattenMappingTool.h b/Engine/Plugins/Enterprise/StaticMeshEditorExtension/Source/StaticMeshEditorExtension/Private/UVTools/UVGenerationFlattenMappingTool.h
--- a/Engine/Plugins/Enterprise/StaticMeshEditorExtension/Source/StaticMeshEditorExtension/Private/UVTools/UVGenerationFlattenMappingTool.h
+++ b/Engine/Plugins/Enterprise/StaticMeshEditorExtension/Source/StaticMeshEditorExtension/Private/UVTools/UVGenerationFlattenMappingTool.h
@@ -43,48 +43,9 @@ public:
 	class FUVGenerationFlattenMappingToolbar* Owner;
 };
 
-//Local actions that can be invoked from this toolbar
-class FUVGenerationFlattenMappingCommands : public TCommands<FUVGenerationFlattenMappingCommands>
-{
-public:
-	FUVGenerationFlattenMappingCommands();
-
-	// TCommands<> interface
-	virtual void RegisterCommands() override;
-	// End of TCommands<> interface
-
-public:
-	/** CommandInfo associated with the EditMode button in the toolbar */
-	TSharedPtr<FUICommandInfo> UnwrapUV;
-};
-
 class FUVGenerationFlattenMappingToolbar : public TSharedFromThis<FUVGenerationFlattenMappingToolbar>
 {
 public:
-	FUVGenerationFlattenMappingToolbar();
-	virtual ~FUVGenerationFlattenMappingToolbar();
-
-	/** Add UV unwrapping items to the StaticMesh Editor's toolbar */
-	static void CreateToolbar(FToolBarBuilder& ToolbarBuilder, const TSharedRef<FUICommandList> CommandList, UStaticMesh* StaticMesh);
-
-private:
-	/** Initialize the toolbar */
-	bool Initialize(UStaticMesh* InStaticMesh, const TSharedRef<FUICommandList> CommandList);
-
-	/** Populate the toolbar */
-	void PopulateToolbar(FToolBarBuilder& ToolbarBuilder, const TSharedRef<FUICommandList> CommandList);
-
-	/** Registers and bind the commands to the toolbar */
-	void BindCommands(const TSharedPtr<FUICommandList> CommandList);
-
-	/** Pointer to the edited static mesh */
-	UStaticMesh* StaticMesh;
-
-	/** Pointer to the static mesh editor */
-	IStaticMeshEditor* StaticMeshEditor;
-
-	/** Pointer to the command list to which the commands are bound to */
-	TSharedPtr<FUICommandList> BoundCommandList;
-
-	TStrongObjectPtr<UUVGenerationFlattenMappingToolbarProxyObject> UVGenerationFlattenMappingToolbarProxyObject;
+	/** Add UV unwrapping menu entry item to the StaticMesh Editor's toolbar */
+	static void CreateMenu(FMenuBuilder& ParentMenuBuilder, const TSharedRef<FUICommandList> CommandList, UStaticMesh* InStaticMesh);
 };
\ No newline at end of file
