Integrating CL# 9622977 from Dev-Enterprise
Updating Unwrap UV generation tool interface

diff --git a/Engine/Plugins/Enterprise/DatasmithImporter/Source/DatasmithImporter/Private/DatasmithStaticMeshImporter.cpp b/Engine/Plugins/Enterprise/DatasmithImporter/Source/DatasmithImporter/Private/DatasmithStaticMeshImporter.cpp
--- a/Engine/Plugins/Enterprise/DatasmithImporter/Source/DatasmithImporter/Private/DatasmithStaticMeshImporter.cpp
+++ b/Engine/Plugins/Enterprise/DatasmithImporter/Source/DatasmithImporter/Private/DatasmithStaticMeshImporter.cpp
@@ -29,6 +29,7 @@
 #include "Templates/UniquePtr.h"
 #include "UObject/Package.h"
 #include "UVTools/UVGenerationFlattenMapping.h"
+#include "UVTools/UVGenerationUtils.h"
 
 #define LOCTEXT_NAMESPACE "DatasmithStaticMeshImporter"
 
@@ -197,35 +198,25 @@ bool FDatasmithStaticMeshImporter::PreBuildStaticMesh( UStaticMesh* StaticMesh )
 		FMeshBuildSettings& BuildSettings = StaticMesh->GetSourceModel(LodIndex).BuildSettings;
 		FMeshDescription& MeshDescription = *StaticMesh->GetMeshDescription(LodIndex);
 
-		// We should always have some UV data in channel 0 because it is used in the mesh tangent calculation during the build.
-		if (!DatasmithMeshHelper::HasUVData(MeshDescription, 0))
-		{
-			DatasmithMeshHelper::CreateDefaultUVs(MeshDescription);
-		}
-
 		if (BuildSettings.bGenerateLightmapUVs && !DatasmithMeshHelper::HasUVData(MeshDescription, BuildSettings.SrcLightmapIndex))
 		{
 			//If no UV data exist at the source index we generate unwrapped UVs.
+			//Do this before calling DatasmithMeshHelper::CreateDefaultUVs() as the UVs may be unwrapped at channel 0.
 			UUVGenerationFlattenMapping::GenerateUVs(MeshDescription, BuildSettings.SrcLightmapIndex, BuildSettings.bRemoveDegenerates);
 		}
 
+		// We should always have some UV data in channel 0 because it is used in the mesh tangent calculation during the build.
+		if (!DatasmithMeshHelper::HasUVData(MeshDescription, 0))
+		{
+			DatasmithMeshHelper::CreateDefaultUVs(MeshDescription);
+		}
+
 		if (DatasmithMeshHelper::IsMeshValid(MeshDescription, BuildSettings.BuildScale3D))
 		{
 			if (BuildSettings.bGenerateLightmapUVs)
 			{
 				DatasmithMeshHelper::RequireUVChannel(MeshDescription, BuildSettings.DstLightmapIndex);
-
-				// Determine the absolute minimum lightmap resolution that can be used for packing
-				float ComparisonThreshold = BuildSettings.bRemoveDegenerates ? THRESH_POINTS_ARE_SAME : 0.0f;
-				FOverlappingCorners OverlappingCorners;
-				FMeshDescriptionOperations::FindOverlappingCorners(OverlappingCorners, MeshDescription, ComparisonThreshold);
-
-				// Packing expects at least one texel per chart. This is the absolute minimum to generate valid UVs.
-				int32 ChartCount = FMeshDescriptionOperations::GetUVChartCount(MeshDescription, BuildSettings.SrcLightmapIndex, ELightmapUVVersion::Latest, OverlappingCorners);
-				const int32 AbsoluteMinResolution = 1 << FMath::CeilLogTwo(FMath::Sqrt(ChartCount));
-				const int32 LightmapResolution = FMath::Clamp(BuildSettings.MinLightmapResolution, AbsoluteMinResolution, 512);
-
-				BuildSettings.MinLightmapResolution = LightmapResolution;
+				UVGenerationUtils::SetupGeneratedLightmapUVResolution(StaticMesh, LodIndex);
 			}
 
 			UStaticMesh::FCommitMeshDescriptionParams Params;
@@ -315,7 +306,7 @@ void FDatasmithStaticMeshImporter::PreBuildStaticMeshes( FDatasmithImportContext
 {
 	TRACE_CPUPROFILER_EVENT_SCOPE(FDatasmithStaticMeshImporter::PreBuildStaticMeshes);
 
-	FScopedSlowTask Progress(ImportContext.ImportedStaticMeshes.Num(), LOCTEXT("PreBuildStaticMeshes", "Packing UVs and computing tangents..."), true, *ImportContext.Warn);
+	FScopedSlowTask Progress(ImportContext.ImportedStaticMeshes.Num(), LOCTEXT("PreBuildStaticMeshes", "Setting up UVs..."), true, *ImportContext.Warn);
 	Progress.MakeDialog(true);
 
 	IMeshUtilities& MeshUtilities = FModuleManager::Get().LoadModuleChecked< IMeshUtilities >( "MeshUtilities" );
@@ -506,24 +497,12 @@ void FDatasmithStaticMeshImporter::SetupStaticMesh( FDatasmithAssetsImportContex
 
 		FMeshDescription& MeshDescription = *StaticMesh->GetMeshDescription(LodIndex);
 
-		auto GetNextOpenUVChannel = [](FMeshDescription& Mesh, int32 FromUVChannelIndex)
-		{
-			for (int32 UVChannelIndex = FromUVChannelIndex; UVChannelIndex < MAX_MESH_TEXTURE_COORDS_MD; ++UVChannelIndex)
-			{
-				if (!DatasmithMeshHelper::HasUVData(Mesh, UVChannelIndex))
-				{
-					return UVChannelIndex;
-				}
-			}
-			return -1;
-		};
-
 		// UV Channels
 		int32 SourceIndex = 0;
 		int32 DestinationIndex = 1;
 		bool bUseImportedLightmap = false;
 		bool bGenerateLightmapUVs = StaticMeshImportOptions.bGenerateLightmapUVs;
-		const int32 FirstOpenUVChannel = GetNextOpenUVChannel(MeshDescription, 0);
+		const int32 FirstOpenUVChannel = UVGenerationUtils::GetNextOpenUVChannel(StaticMesh, LodIndex);
 
 		// if a custom lightmap coordinate index was imported, disable lightmap generation
 		if (DatasmithMeshHelper::HasUVData(MeshDescription, MeshElement->GetLightmapCoordinateIndex()))
diff --git a/Engine/Plugins/Enterprise/StaticMeshEditorExtension/Source/StaticMeshEditorExtension/Private/StaticMeshEditorExtensionModule.cpp b/Engine/Plugins/Enterprise/StaticMeshEditorExtension/Source/StaticMeshEditorExtension/Private/StaticMeshEditorExtensionModule.cpp
--- a/Engine/Plugins/Enterprise/StaticMeshEditorExtension/Source/StaticMeshEditorExtension/Private/StaticMeshEditorExtensionModule.cpp
+++ b/Engine/Plugins/Enterprise/StaticMeshEditorExtension/Source/StaticMeshEditorExtension/Private/StaticMeshEditorExtensionModule.cpp
@@ -7,7 +7,7 @@
 #include "UVTools/UVGenerationFlattenMappingTool.h"
 #include "UVTools/UVGenerationToolbar.h"
 #include "UVTools/UVGenerationTool.h"
-#include "UVTools/UVGenerationUtils.h"
+#include "UVTools/UVGenerationSettings.h"
 
 #include "ContentBrowserModule.h"
 #include "Modules/ModuleManager.h"
diff --git a/Engine/Plugins/Enterprise/StaticMeshEditorExtension/Source/StaticMeshEditorExtension/Private/UVTools/SUVGenerationTool.h b/Engine/Plugins/Enterprise/StaticMeshEditorExtension/Source/StaticMeshEditorExtension/Private/UVTools/SUVGenerationTool.h
--- a/Engine/Plugins/Enterprise/StaticMeshEditorExtension/Source/StaticMeshEditorExtension/Private/UVTools/SUVGenerationTool.h
+++ b/Engine/Plugins/Enterprise/StaticMeshEditorExtension/Source/StaticMeshEditorExtension/Private/UVTools/SUVGenerationTool.h
@@ -8,7 +8,7 @@
 #include "Widgets/DeclarativeSyntaxSupport.h"
 #include "Widgets/SCompoundWidget.h"
 #include "Styling/SlateColor.h"
-#include "UVGenerationUtils.h"
+#include "UVGenerationSettings.h"
 #include "SUVGenerationTool.generated.h"
 
 class SMultiLineEditableTextBox;
diff --git a/Engine/Plugins/Enterprise/StaticMeshEditorExtension/Source/StaticMeshEditorExtension/Private/UVTools/UVGenerationDataprepOperation.cpp b/Engine/Plugins/Enterprise/StaticMeshEditorExtension/Source/StaticMeshEditorExtension/Private/UVTools/UVGenerationDataprepOperation.cpp
--- a/Engine/Plugins/Enterprise/StaticMeshEditorExtension/Source/StaticMeshEditorExtension/Private/UVTools/UVGenerationDataprepOperation.cpp
+++ b/Engine/Plugins/Enterprise/StaticMeshEditorExtension/Source/StaticMeshEditorExtension/Private/UVTools/UVGenerationDataprepOperation.cpp
@@ -3,15 +3,19 @@
 #include "UVGenerationDataprepOperation.h"
 
 #include "UVTools/UVGenerationFlattenMapping.h"
+#include "UVTools/UVGenerationUtils.h"
 
+#include "Async/ParallelFor.h"
 #include "Components/StaticMeshComponent.h"
 #include "Engine/StaticMesh.h"
 #include "Engine/StaticMeshActor.h"
-#include "MeshDescriptionOperations.h"
+#include "MeshDescription.h"
 #include "OverlappingCorners.h"
 
 void UUVGenerationFlattenMappingOperation::OnExecution_Implementation(const FDataprepContext& InContext)
 {
+	TRACE_CPUPROFILER_EVENT_SCOPE(UUVGenerationFlattenMappingOperation::OnExecution_Implementation)
+
 	TSet<UStaticMesh*> StaticMeshes;
 
 	for (UObject* Object : InContext.Objects)
@@ -46,7 +50,8 @@ void UUVGenerationFlattenMappingOperation::OnExecution_Implementation(const FDat
 		}
 	}
 
-	for (UStaticMesh* StaticMesh : StaticMeshes)
+	// The actual work
+	auto Compute = [&](UStaticMesh* StaticMesh)
 	{
 		for (int32 LodIndex = 0; LodIndex < StaticMesh->GetNumSourceModels(); ++LodIndex)
 		{
@@ -55,10 +60,36 @@ void UUVGenerationFlattenMappingOperation::OnExecution_Implementation(const FDat
 				continue;
 			}
 
+			int32 TargetChannel = UVChannel;
+			if (ChannelSelection == EUnwrappedUVDatasmithOperationChannelSelection::FirstEmptyChannel)
+			{
+				TargetChannel = UVGenerationUtils::GetNextOpenUVChannel(StaticMesh, LodIndex);
+				if (TargetChannel < 0)
+				{
+					continue;
+				}
+			}
+
 			FMeshBuildSettings& BuildSettings = StaticMesh->GetSourceModel(LodIndex).BuildSettings;
 			FMeshDescription& MeshDescription = *StaticMesh->GetMeshDescription(LodIndex);
 
-			UUVGenerationFlattenMapping::GenerateUVs(MeshDescription, UVChannel, BuildSettings.bRemoveDegenerates, AngleThreshold, AreaWeight);
+			UUVGenerationFlattenMapping::GenerateUVs(MeshDescription, TargetChannel, BuildSettings.bRemoveDegenerates, AngleThreshold);
 		}
-	}
+	};
+
+	// Start with the biggest mesh first to help balancing tasks on threads
+	TArray<UStaticMesh*> StaticMeshArray(StaticMeshes.Array());
+	Algo::SortBy(
+		StaticMeshArray,
+		[](const UStaticMesh* Mesh) { return Mesh->IsMeshDescriptionValid(0) ? Mesh->GetMeshDescription(0)->Vertices().Num() : 0; },
+		TGreater<>()
+	);
+
+	ParallelFor(StaticMeshArray.Num(),
+		[&](int32 Index)
+		{
+			Compute(StaticMeshArray[Index]);
+		},
+		EParallelForFlags::Unbalanced
+	);
 }
\ No newline at end of file
diff --git a/Engine/Plugins/Enterprise/StaticMeshEditorExtension/Source/StaticMeshEditorExtension/Private/UVTools/UVGenerationDataprepOperation.h b/Engine/Plugins/Enterprise/StaticMeshEditorExtension/Source/StaticMeshEditorExtension/Private/UVTools/UVGenerationDataprepOperation.h
--- a/Engine/Plugins/Enterprise/StaticMeshEditorExtension/Source/StaticMeshEditorExtension/Private/UVTools/UVGenerationDataprepOperation.h
+++ b/Engine/Plugins/Enterprise/StaticMeshEditorExtension/Source/StaticMeshEditorExtension/Private/UVTools/UVGenerationDataprepOperation.h
@@ -7,28 +7,35 @@
 
 #include "UVGenerationDataprepOperation.generated.h"
 
-UCLASS(Experimental, Category = MeshOperation, Meta = (DisplayName = "Generate Flatten Mapping UVs", ToolTip = "For each static mesh to process, generate a flat UV map in the specified channel"))
+UENUM()
+enum class EUnwrappedUVDatasmithOperationChannelSelection : uint8
+{
+	FirstEmptyChannel UMETA(Tooltip = "Generate the unwrapped UV in the first UV channel that is empty."),
+	SpecifyChannel UMETA(Tooltip = "Manually select the target UV channel for the unwrapped UV generation."),
+};
+
+UCLASS(Experimental, Category = MeshOperation, Meta = (DisplayName = "Generate Unwrapped UVs", ToolTip = "For each static mesh to process, generate an unwrapped UV map in the specified channel"))
 class STATICMESHEDITOREXTENSION_API UUVGenerationFlattenMappingOperation : public UDataprepOperation
 {
 	GENERATED_BODY()
 
 	UUVGenerationFlattenMappingOperation()
-		: UVChannel(0),
-		AngleThreshold(66.f),
-		AreaWeight(0.7f)
+		: ChannelSelection(EUnwrappedUVDatasmithOperationChannelSelection::FirstEmptyChannel),
+		UVChannel(0),
+		AngleThreshold(66.f)
 	{
 	}
 
 public:
-	UPROPERTY(EditAnywhere, BlueprintReadWrite, Category = "UV Generation Settings", meta = (ToolTip="The UV channel where to generate the flatten mapping", ClampMin = "0", ClampMax = "7")) //Clampmax is from MAX_MESH_TEXTURE_COORDS_MD - 1
+	UPROPERTY(EditAnywhere, BlueprintReadWrite, Category = "UV Generation Settings")
+	EUnwrappedUVDatasmithOperationChannelSelection ChannelSelection;
+
+	UPROPERTY(EditAnywhere, BlueprintReadWrite, Category = "UV Generation Settings", meta = (DisplayName = "UV Channel", ToolTip = "The UV channel where to generate the flatten mapping", ClampMin = "0", ClampMax = "7", EditCondition = "ChannelSelection == EUnwrappedUVDatasmithOperationChannelSelection::SpecifyChannel")) //Clampmax is from MAX_MESH_TEXTURE_COORDS_MD - 1
 	int UVChannel;
 
 	UPROPERTY(EditAnywhere, BlueprintReadWrite, Category = "UV Generation Settings", meta = (ClampMin = "1", ClampMax = "90"))
 	float AngleThreshold;
 
-	UPROPERTY(EditAnywhere, BlueprintReadWrite, Category = "UV Generation Settings", meta = (ClampMin = "0", ClampMax = "1"))
-	float AreaWeight;
-
 	//~ Begin UDataprepOperation Interface
 public:
 	virtual FText GetCategory_Implementation() const override
diff --git a/Engine/Plugins/Enterprise/StaticMeshEditorExtension/Source/StaticMeshEditorExtension/Private/UVTools/UVGenerationFlattenMapping.cpp b/Engine/Plugins/Enterprise/StaticMeshEditorExtension/Source/StaticMeshEditorExtension/Private/UVTools/UVGenerationFlattenMapping.cpp
--- a/Engine/Plugins/Enterprise/StaticMeshEditorExtension/Source/StaticMeshEditorExtension/Private/UVTools/UVGenerationFlattenMapping.cpp
+++ b/Engine/Plugins/Enterprise/StaticMeshEditorExtension/Source/StaticMeshEditorExtension/Private/UVTools/UVGenerationFlattenMapping.cpp
@@ -21,6 +21,8 @@
 #include "StaticMeshAttributes.h"
 #include "Templates/TypeHash.h"
 
+#define FLATTEN_AREA_WEIGHT 0.7
+
 class FUVGenerationFlattenMappingInternal
 {
 public:
@@ -1152,7 +1154,7 @@ TArray<FUVGenerationFlattenMappingInternal::FaceStruct> FUVGenerationFlattenMapp
 	return FlattenFaces;
 }
 
-void UUVGenerationFlattenMapping::GenerateFlattenMappingUVs(UStaticMesh* InStaticMesh, int32 UVChannel, float AngleThreshold, float AreaWeight)
+void UUVGenerationFlattenMapping::GenerateFlattenMappingUVs(UStaticMesh* InStaticMesh, int32 UVChannel, float AngleThreshold)
 {
 	if (InStaticMesh == nullptr)
 	{
@@ -1173,7 +1175,7 @@ void UUVGenerationFlattenMapping::GenerateFlattenMappingUVs(UStaticMesh* InStati
 
 		FMeshBuildSettings& BuildSettings = InStaticMesh->GetSourceModel(LodIndex).BuildSettings;
 		FMeshDescription& MeshDescription = *InStaticMesh->GetMeshDescription(LodIndex);
-		UUVGenerationFlattenMapping::GenerateUVs(MeshDescription, UVChannel, BuildSettings.bRemoveDegenerates, AngleThreshold, AreaWeight);
+		UUVGenerationFlattenMapping::GenerateUVs(MeshDescription, UVChannel, BuildSettings.bRemoveDegenerates, AngleThreshold);
 
 		UStaticMesh::FCommitMeshDescriptionParams CommitMeshDescriptionParam;
 		CommitMeshDescriptionParam.bUseHashAsGuid = true;
@@ -1183,7 +1185,7 @@ void UUVGenerationFlattenMapping::GenerateFlattenMappingUVs(UStaticMesh* InStati
 	InStaticMesh->PostEditChange();
 }
 
-void UUVGenerationFlattenMapping::GenerateUVs(FMeshDescription& InMesh, int32 UVChannel, bool bRemoveDegenerates, float AngleThreshold, float AreaWeight)
+void UUVGenerationFlattenMapping::GenerateUVs(FMeshDescription& InMesh, int32 UVChannel, bool bRemoveDegenerates, float AngleThreshold)
 {
 	if (!ensure(UVChannel >= 0 && UVChannel < MAX_MESH_TEXTURE_COORDS_MD))
 	{
@@ -1202,7 +1204,7 @@ void UUVGenerationFlattenMapping::GenerateUVs(FMeshDescription& InMesh, int32 UV
 	// Get the Internal Face Array
 	TArray< FUVGenerationFlattenMappingInternal::FaceStruct > FlattenFaces = FUVGenerationFlattenMappingInternal::GetFacesFrom(InMesh, OverlappingCorners);
 	
-	FUVGenerationFlattenMappingInternal::CalculateUVs(FlattenFaces, InMesh.Vertices().Num(), AngleThreshold, AreaWeight);
+	FUVGenerationFlattenMappingInternal::CalculateUVs(FlattenFaces, InMesh.Vertices().Num(), AngleThreshold, FLATTEN_AREA_WEIGHT);
 
 	// Write unwrapped UVs back to mesh description
 	// Vertex instances shared across islands will be split
diff --git a/Engine/Plugins/Enterprise/StaticMeshEditorExtension/Source/StaticMeshEditorExtension/Private/UVTools/UVGenerationFlattenMappingTool.cpp b/Engine/Plugins/Enterprise/StaticMeshEditorExtension/Source/StaticMeshEditorExtension/Private/UVTools/UVGenerationFlattenMappingTool.cpp
--- a/Engine/Plugins/Enterprise/StaticMeshEditorExtension/Source/StaticMeshEditorExtension/Private/UVTools/UVGenerationFlattenMappingTool.cpp
+++ b/Engine/Plugins/Enterprise/StaticMeshEditorExtension/Source/StaticMeshEditorExtension/Private/UVTools/UVGenerationFlattenMappingTool.cpp
@@ -16,9 +16,12 @@
 #include "Interfaces/IMainFrameModule.h"
 #include "Interfaces/IPluginManager.h"
 #include "IStaticMeshEditor.h"
+#include "MeshDescriptionOperations.h"
 #include "MeshEditorUtils.h"
+#include "MeshUtilitiesCommon.h"
 #include "Misc/ScopedSlowTask.h"
 #include "Modules/ModuleManager.h"
+#include "OverlappingCorners.h"
 #include "PropertyEditorModule.h"
 #include "ScopedTransaction.h"
 #include "Styling/SlateStyle.h"
@@ -33,6 +36,7 @@
 
 #define LOCTEXT_NAMESPACE "UVGenerationFlattenMappingTool"
 
+DEFINE_LOG_CATEGORY_STATIC(LogUVUnwrapping, Log, All);
 
 #define IMAGE_PLUGIN_BRUSH( RelativePath, ... ) FSlateImageBrush( FUVGenerationFlattenMappingToolStyle::InContent( RelativePath, ".png" ), __VA_ARGS__ )
 
@@ -105,12 +109,12 @@ public:
 	SLATE_BEGIN_ARGS(SUVGenerationFlattenMappingWindow)
 	{}
 		SLATE_ARGUMENT(TArray<UStaticMesh*>*, StaticMeshes)
-		SLATE_ARGUMENT(UUVFlattenMappingSettings*, MappingSettings)
+		SLATE_ARGUMENT(UUVUnwrapSettings*, UnwrapSettings)
 		SLATE_ARGUMENT(TSharedPtr<SWindow>, WidgetWindow)
 	SLATE_END_ARGS()
 
 public:
-	static bool DisplayDialog(TArray<UStaticMesh*>* StaticMeshes, UUVFlattenMappingSettings* OutFlattenMappingSettings);
+	static bool DisplayDialog(TArray<UStaticMesh*>* StaticMeshes, UUVUnwrapSettings* OutUVUnwrapSettings);
 
 	void Construct(const FArguments& InArgs);
 
@@ -150,7 +154,7 @@ private:
 
 private:
 	TArray<UStaticMesh*>* StaticMeshes;
-	UUVFlattenMappingSettings* MappingSettings;
+	UUVUnwrapSettings* UnwrapSettings;
 	TWeakPtr< SWindow > Window;
 	bool bCanProceed;
 };
@@ -193,13 +197,13 @@ TSharedRef<FExtender> FUVGenerationFlattenMappingTool::OnExtendContentBrowserAss
 
 void FUVGenerationFlattenMappingTool::OpenUnwrapUVWindow(TArray<UStaticMesh*> StaticMeshes)
 {
-	TStrongObjectPtr<UUVFlattenMappingSettings> MappingSettings = TStrongObjectPtr<UUVFlattenMappingSettings>(NewObject<UUVFlattenMappingSettings>(GetTransientPackage(), TEXT("Flatten Mapping UV Generation Settings")));
+	TStrongObjectPtr<UUVUnwrapSettings> UVUnwrapSettings = TStrongObjectPtr<UUVUnwrapSettings>(NewObject<UUVUnwrapSettings>(GetTransientPackage(), TEXT("Flatten Mapping UV Generation Settings")));
 	bool bCancelled = false;
 
-	if (SUVGenerationFlattenMappingWindow::DisplayDialog(&StaticMeshes, MappingSettings.Get()))
+	if (SUVGenerationFlattenMappingWindow::DisplayDialog(&StaticMeshes, UVUnwrapSettings.Get()))
 	{
 		// Save parameters to config file
-		MappingSettings->SaveConfig(CPF_Config, *MappingSettings->GetDefaultConfigFilename());
+		UVUnwrapSettings->SaveConfig(CPF_Config, *UVUnwrapSettings->GetDefaultConfigFilename());
 
 		FScopedTransaction Transaction(LOCTEXT("GenerateUnwrappedUVsTransation", "Generate Unwrapped UVs"));
 		FText SlowTaskProgressText(LOCTEXT("UnwrappingUVsSlowTask", "Unwrapping UVs ({0}/{1})"));
@@ -216,7 +220,33 @@ void FUVGenerationFlattenMappingTool::OpenUnwrapUVWindow(TArray<UStaticMesh*> St
 			};
 
 			SlowTask.EnterProgressFrame(1, FText::Format(SlowTaskProgressText, SlowTask.CompletedWork + 1, SlowTask.TotalAmountOfWork));
-			UUVGenerationFlattenMapping::GenerateFlattenMappingUVs(CurrentStaticMesh, MappingSettings->UVChannel, MappingSettings->AngleThreshold, MappingSettings->AreaWeight);
+			CurrentStaticMesh->Modify();
+
+			for (int32 LODIndex = 0; LODIndex < CurrentStaticMesh->GetNumSourceModels(); ++LODIndex)
+			{
+				int32 UVChannel = -1;
+
+				if (SetupMeshForUVGeneration(CurrentStaticMesh, UVUnwrapSettings.Get(), LODIndex, UVChannel))
+				{
+					UUVGenerationFlattenMapping::GenerateFlattenMappingUVs(CurrentStaticMesh, UVChannel, UVUnwrapSettings->AngleThreshold);
+
+					if (UVUnwrapSettings->ChannelSelection == EUnwrappedUVChannelSelection::AutomaticLightmapSetup)
+					{
+						UVGenerationUtils::SetupGeneratedLightmapUVResolution(CurrentStaticMesh, LODIndex);
+					}
+					
+					UStaticMesh::FCommitMeshDescriptionParams CommitMeshDescriptionParam;
+					CommitMeshDescriptionParam.bUseHashAsGuid = true;
+					CurrentStaticMesh->CommitMeshDescription(LODIndex, CommitMeshDescriptionParam);
+				}
+				else
+				{
+					UE_LOG(LogUVUnwrapping, Error, TEXT("Could not generate unwrapped UV at the specified channel for static mesh %s"), *CurrentStaticMesh->GetName());
+					break;
+				}
+			}
+
+			CurrentStaticMesh->PostEditChange();
 
 			if (SlowTask.ShouldCancel())
 			{
@@ -233,7 +263,60 @@ void FUVGenerationFlattenMappingTool::OpenUnwrapUVWindow(TArray<UStaticMesh*> St
 	}
 }
 
-bool SUVGenerationFlattenMappingWindow::DisplayDialog(TArray<UStaticMesh*>* StaticMeshes, UUVFlattenMappingSettings* OutFlattenMappingSettings)
+bool FUVGenerationFlattenMappingTool::SetupMeshForUVGeneration(UStaticMesh* StaticMesh, const UUVUnwrapSettings* UVUnwrapSettings, int32 LODIndex, int32& OutChannelIndex)
+{
+	OutChannelIndex = -1;
+
+	if (UVUnwrapSettings->ChannelSelection == EUnwrappedUVChannelSelection::SpecifyChannel)
+	{
+		if (UVUnwrapSettings->UVChannel >= 0 && UVUnwrapSettings->UVChannel < MAX_MESH_TEXTURE_COORDS_MD)
+		{
+			OutChannelIndex = UVUnwrapSettings->UVChannel;
+		}
+	}
+	else if (UVUnwrapSettings->ChannelSelection == EUnwrappedUVChannelSelection::AutomaticLightmapSetup)
+	{
+		FMeshBuildSettings& BuildSettings = StaticMesh->GetSourceModel(LODIndex).BuildSettings;
+
+		if (!BuildSettings.bGenerateLightmapUVs)
+		{
+			//If the lightmap generation was deactivated we change source and destination indexes to the first empty slot.
+			int32 FirstOpenChannel = UVGenerationUtils::GetNextOpenUVChannel(StaticMesh, LODIndex);
+
+			if (FirstOpenChannel >= 0)
+			{
+				BuildSettings.SrcLightmapIndex = FirstOpenChannel;
+				BuildSettings.DstLightmapIndex = FirstOpenChannel;
+				OutChannelIndex = FirstOpenChannel;
+
+				if (LODIndex == 0)
+				{
+					//If we are setting up the first LOD make sure the mesh lightmap coordinate points to the generated ones.
+					StaticMesh->LightMapCoordinateIndex = FirstOpenChannel;
+				}
+
+				BuildSettings.bGenerateLightmapUVs = true;
+			}
+		}
+		else
+		{
+			OutChannelIndex = BuildSettings.SrcLightmapIndex;
+		}
+	}
+	else if ( UVUnwrapSettings->ChannelSelection == EUnwrappedUVChannelSelection::FirstEmptyChannel )
+	{
+		int32 FirstOpenChannel = UVGenerationUtils::GetNextOpenUVChannel(StaticMesh, LODIndex);
+
+		if (FirstOpenChannel >= 0)
+		{
+			OutChannelIndex = FirstOpenChannel;
+		}
+	}
+
+	return OutChannelIndex >= 0;
+}
+
+bool SUVGenerationFlattenMappingWindow::DisplayDialog(TArray<UStaticMesh*>* StaticMeshes, UUVUnwrapSettings* OutUVUnwrapSettings)
 {
 	TSharedPtr<SWindow> ParentWindow;
 
@@ -252,7 +335,7 @@ bool SUVGenerationFlattenMappingWindow::DisplayDialog(TArray<UStaticMesh*>* Stat
 	(
 		SAssignNew(ParameterWindow, SUVGenerationFlattenMappingWindow)
 		.StaticMeshes(StaticMeshes)
-		.MappingSettings(OutFlattenMappingSettings)
+		.UnwrapSettings(OutUVUnwrapSettings)
 		.WidgetWindow(Window)
 	);
 
@@ -264,7 +347,7 @@ bool SUVGenerationFlattenMappingWindow::DisplayDialog(TArray<UStaticMesh*>* Stat
 void SUVGenerationFlattenMappingWindow::Construct(const FArguments& InArgs)
 {
 	StaticMeshes = InArgs._StaticMeshes;
-	MappingSettings = InArgs._MappingSettings;
+	UnwrapSettings = InArgs._UnwrapSettings;
 	Window = InArgs._WidgetWindow;
 	bCanProceed = false;
 
@@ -320,7 +403,7 @@ void SUVGenerationFlattenMappingWindow::Construct(const FArguments& InArgs)
 	DetailsViewArgs.NameAreaSettings = FDetailsViewArgs::HideNameArea;
 	TSharedPtr<IDetailsView> DetailsView = PropertyEditorModule.CreateDetailView(DetailsViewArgs);
 
-	DetailsView->SetObject(MappingSettings);
+	DetailsView->SetObject(UnwrapSettings);
 	DetailsViewBox->SetContent(DetailsView.ToSharedRef());
 }
 
diff --git a/Engine/Plugins/Enterprise/StaticMeshEditorExtension/Source/StaticMeshEditorExtension/Private/UVTools/UVGenerationFlattenMappingTool.h b/Engine/Plugins/Enterprise/StaticMeshEditorExtension/Source/StaticMeshEditorExtension/Private/UVTools/UVGenerationFlattenMappingTool.h
--- a/Engine/Plugins/Enterprise/StaticMeshEditorExtension/Source/StaticMeshEditorExtension/Private/UVTools/UVGenerationFlattenMappingTool.h
+++ b/Engine/Plugins/Enterprise/StaticMeshEditorExtension/Source/StaticMeshEditorExtension/Private/UVTools/UVGenerationFlattenMappingTool.h
@@ -18,6 +18,30 @@ class UStaticMesh;
 class IStaticMeshEditor;
 struct FAssetData;
 
+UENUM()
+enum class EUnwrappedUVChannelSelection : uint8
+{
+	AutomaticLightmapSetup UMETA(Tooltip = "Enable lightmap generation and use the generated unwrapped UV as the lightmap source."),
+	FirstEmptyChannel UMETA(Tooltip = "Generate the unwrapped UV in the first UV channel that is empty."),
+	SpecifyChannel UMETA(Tooltip = "Manually select the target UV channel for the unwrapped UV generation."),
+};
+
+UCLASS(config = EditorPerProjectUserSettings, Transient)
+class UUVUnwrapSettings : public UObject
+{
+	GENERATED_BODY()
+
+public:
+	UPROPERTY(config, EditAnywhere, Category = "Flatten Mapping")
+	EUnwrappedUVChannelSelection ChannelSelection = EUnwrappedUVChannelSelection::AutomaticLightmapSetup;
+
+	UPROPERTY(config, EditAnywhere, category = "Flatten Mapping", meta = (ToolTip = "The UV channel where to generate the flatten mapping", ClampMin = "0", ClampMax = "7", EditCondition = "ChannelSelection == EUnwrappedUVChannelSelection::SpecifyChannel")) //Clampmax is from MAX_MESH_TEXTURE_COORDS_MD - 1
+	int32 UVChannel = 0;
+
+	UPROPERTY(config, EditAnywhere, category = "Flatten Mapping", meta = (ClampMin = "1", ClampMax = "90"))
+	float AngleThreshold = 66.f;
+};
+
 class FUVGenerationFlattenMappingTool
 {
 public:
@@ -30,6 +54,20 @@ public:
 	 * Opens the Unwrap UV dialog window.
 	 */
 	static void OpenUnwrapUVWindow(TArray<UStaticMesh*> StaticMeshes);
+
+private:
+
+	/*
+	 * Setup the static mesh for UV generation with the given UVUnwrapSettings, enabling bGenerateLightmap if needed and selecting the proper UVChannel according to the options.
+	 * 
+	 * @param StaticMesh		The StaticMesh being set up
+	 * @param UVUnwrapSettings	The generation settings used to chose the UV channel index and set up the lightmap generation
+	 * @param LODIndex			The LOD being set up.
+	 * @param OutChannelIndex	Out parameter that gives UV channel we should generate the Unwrapped UV into.
+	 *
+	 * @return Returns true if a UV channel could be safely selected.
+	 */
+	static bool SetupMeshForUVGeneration(UStaticMesh* StaticMesh, const UUVUnwrapSettings* UVUnwrapSettings, int32 LODIndex, int32& OutChannelIndex);
 };
 
 UCLASS()
diff --git a/Engine/Plugins/Enterprise/StaticMeshEditorExtension/Source/StaticMeshEditorExtension/Private/UVTools/UVGenerationSettings.cpp b/Engine/Plugins/Enterprise/StaticMeshEditorExtension/Source/StaticMeshEditorExtension/Private/UVTools/UVGenerationSettings.cpp
new file mode 100644
--- /dev/null
+++ b/Engine/Plugins/Enterprise/StaticMeshEditorExtension/Source/StaticMeshEditorExtension/Private/UVTools/UVGenerationSettings.cpp
@@ -0,0 +1,64 @@
+// Copyright 1998-2019 Epic Games, Inc. All Rights Reserved.
+
+#include "UVTools/UVGenerationSettings.h"
+
+#include "DetailWidgetRow.h"
+#include "Engine/StaticMesh.h"
+#include "IDetailChildrenBuilder.h"
+#include "IDetailPropertyRow.h"
+#include "Math/Axis.h"
+#include "MeshDescription.h"
+#include "MeshUtilitiesCommon.h"
+#include "PropertyHandle.h"
+#include "StaticMeshAttributes.h"
+#include "Widgets/Input/SVectorInputBox.h"
+
+TSharedRef<IPropertyTypeCustomization> FUVGenerationSettingsCustomization::MakeInstance()
+{
+	return MakeShareable(new FUVGenerationSettingsCustomization);
+}
+
+void FUVGenerationSettingsCustomization::CustomizeHeader(TSharedRef<IPropertyHandle> PropertyHandle, FDetailWidgetRow& HeaderRow, IPropertyTypeCustomizationUtils& CustomizationUtils)
+{
+	TArray<void*> StructPtrs;
+	PropertyHandle->AccessRawData(StructPtrs);
+	GenerateUVSettings = (StructPtrs.Num() == 1) ? reinterpret_cast<FUVGenerationSettings*>(StructPtrs[0]) : nullptr;
+}
+
+void FUVGenerationSettingsCustomization::CustomizeChildren(TSharedRef<IPropertyHandle> StructPropertyHandle, IDetailChildrenBuilder& ChildBuilder, IPropertyTypeCustomizationUtils& StructCustomizationUtils)
+{
+	uint32 NumChildProps = 0;
+
+	StructPropertyHandle->GetNumChildren(NumChildProps);
+
+	for (uint32 Idx = 0; Idx < NumChildProps; Idx++)
+	{
+		TSharedPtr<IPropertyHandle> PropHandle = StructPropertyHandle->GetChildHandle(Idx);
+
+		FName PropertyName = PropHandle->GetProperty()->GetFName();
+		if (PropertyName == GET_MEMBER_NAME_CHECKED(FUVGenerationSettings, Size)
+			|| PropertyName == GET_MEMBER_NAME_CHECKED(FUVGenerationSettings, Position)
+			|| PropertyName == GET_MEMBER_NAME_CHECKED(FUVGenerationSettings, Rotation))
+		{
+			PropHandle->SetOnChildPropertyValueChanged(FSimpleDelegate::CreateSP(this, &FUVGenerationSettingsCustomization::OnShapePropertyChanged));
+		}
+		else if (PropertyName == GET_MEMBER_NAME_CHECKED(FUVGenerationSettings, TargetChannel))
+		{
+			if (GenerateUVSettings && GenerateUVSettings->OnGetNumberOfUVs.IsBound())
+			{
+				PropHandle->SetInstanceMetaData(TEXT("ClampMax"), FString::FromInt(FMath::Min(GenerateUVSettings->OnGetNumberOfUVs.Execute(), (int32)MAX_MESH_TEXTURE_COORDS_MD - 1)));
+			}
+			else
+			{
+				PropHandle->SetInstanceMetaData(TEXT("ClampMax"), FString::FromInt(MAX_MESH_TEXTURE_COORDS_MD - 1));
+			}
+		}
+
+		IDetailPropertyRow& PropertyRow = ChildBuilder.AddProperty(PropHandle.ToSharedRef());
+	}
+}
+
+void FUVGenerationSettingsCustomization::OnShapePropertyChanged()
+{
+	GenerateUVSettings->OnShapeEditingValueChanged.Broadcast();
+}
\ No newline at end of file
diff --git a/Engine/Plugins/Enterprise/StaticMeshEditorExtension/Source/StaticMeshEditorExtension/Private/UVTools/UVGenerationSettings.h b/Engine/Plugins/Enterprise/StaticMeshEditorExtension/Source/StaticMeshEditorExtension/Private/UVTools/UVGenerationSettings.h
new file mode 100644
--- /dev/null
+++ b/Engine/Plugins/Enterprise/StaticMeshEditorExtension/Source/StaticMeshEditorExtension/Private/UVTools/UVGenerationSettings.h
@@ -0,0 +1,80 @@
+// Copyright 1998-2019 Epic Games, Inc. All Rights Reserved.
+
+#pragma once
+
+#include "CoreMinimal.h"
+#include "IPropertyTypeCustomization.h"
+
+#include "UVGenerationSettings.generated.h"
+
+UENUM()
+enum class EGenerateUVProjectionType : uint8
+{
+	Box,
+	Cylindrical,
+	Planar,
+};
+
+USTRUCT()
+struct FUVGenerationSettings
+{
+	GENERATED_BODY()
+
+public:
+	UPROPERTY(EditAnywhere, category = "Projection Settings")
+	EGenerateUVProjectionType ProjectionType;
+	UPROPERTY(EditAnywhere, category = "Projection Settings")
+	FVector Position;
+	UPROPERTY(EditAnywhere, category = "Projection Settings")
+	FRotator Rotation;
+	UPROPERTY(EditAnywhere, category = "Projection Settings")
+	FVector Size;
+	UPROPERTY(EditAnywhere, category = "Projection Settings", meta = (DisplayName = "UV Tiling Scale"))
+	FVector2D UVTilingScale = FVector2D::UnitVector;
+	UPROPERTY(EditAnywhere, category = "Projection Settings", meta = (DisplayName = "UV Offset"))
+	FVector2D UVOffset = FVector2D::ZeroVector;
+
+	//ClampMax defined by MAX_MESH_TEXTURE_COORDS
+	UPROPERTY(EditAnywhere, category = "Projection Settings", meta = (ClampMin = 0, ClampMax = 7, Tooltip = "The UV channel the projection will be applied to."))
+	uint8 TargetChannel;
+
+	//Used to get access to the StaticMeshEditor without creating a dependency.
+	DECLARE_DELEGATE_RetVal(int32, FOnGetNumberOfUVs)
+	FOnGetNumberOfUVs OnGetNumberOfUVs;
+
+	//Using a delegate instead of an event because events assume the owning type is a class.
+	DECLARE_MULTICAST_DELEGATE(FOnShapeEditingValuesChanged)
+	FOnShapeEditingValuesChanged OnShapeEditingValueChanged;
+
+	FUVGenerationSettings()
+		: ProjectionType(EGenerateUVProjectionType::Box)
+		, Position(FVector::ZeroVector)
+		, Rotation(0)
+		, Size(FVector::OneVector)
+		, UVTilingScale(FVector2D::UnitVector)
+		, UVOffset(FVector2D::ZeroVector)
+		, TargetChannel(0)
+	{}
+};
+
+class FUVGenerationSettingsCustomization : public IPropertyTypeCustomization
+{
+	enum class EShapeField {
+		Size,
+		Position,
+		Rotation
+	};
+
+public:
+	/** Makes a new instance of this detail layout class for a specific detail view requesting it */
+	static TSharedRef<IPropertyTypeCustomization> MakeInstance();
+
+	virtual void CustomizeHeader(TSharedRef<IPropertyHandle> PropertyHandle, FDetailWidgetRow& HeaderRow, IPropertyTypeCustomizationUtils& CustomizationUtils) override;
+	virtual void CustomizeChildren(TSharedRef<IPropertyHandle> StructPropertyHandle, class IDetailChildrenBuilder& ChildBuilder, IPropertyTypeCustomizationUtils& StructCustomizationUtils) override;
+
+private:
+
+	void OnShapePropertyChanged();
+
+	FUVGenerationSettings* GenerateUVSettings = nullptr;
+};
\ No newline at end of file
diff --git a/Engine/Plugins/Enterprise/StaticMeshEditorExtension/Source/StaticMeshEditorExtension/Private/UVTools/UVGenerationTool.cpp b/Engine/Plugins/Enterprise/StaticMeshEditorExtension/Source/StaticMeshEditorExtension/Private/UVTools/UVGenerationTool.cpp
--- a/Engine/Plugins/Enterprise/StaticMeshEditorExtension/Source/StaticMeshEditorExtension/Private/UVTools/UVGenerationTool.cpp
+++ b/Engine/Plugins/Enterprise/StaticMeshEditorExtension/Source/StaticMeshEditorExtension/Private/UVTools/UVGenerationTool.cpp
@@ -3,6 +3,7 @@
 #include "UVGenerationTool.h"
 #include "EditorModeManager.h"
 #include "EditorViewportClient.h"
+#include "UVGenerationSettings.h"
 
 #define LOCTEXT_NAMESPACE "UVGenerationTool"
 
diff --git a/Engine/Plugins/Enterprise/StaticMeshEditorExtension/Source/StaticMeshEditorExtension/Private/UVTools/UVGenerationTool.h b/Engine/Plugins/Enterprise/StaticMeshEditorExtension/Source/StaticMeshEditorExtension/Private/UVTools/UVGenerationTool.h
--- a/Engine/Plugins/Enterprise/StaticMeshEditorExtension/Source/StaticMeshEditorExtension/Private/UVTools/UVGenerationTool.h
+++ b/Engine/Plugins/Enterprise/StaticMeshEditorExtension/Source/StaticMeshEditorExtension/Private/UVTools/UVGenerationTool.h
@@ -4,7 +4,9 @@
 
 #include "CoreMinimal.h"
 #include "EdMode.h"
-#include "UVGenerationUtils.h"
+
+struct FUVGenerationSettings;
+enum class EGenerateUVProjectionType : uint8;
 
 class FUVGenerationTool : public FEdMode
 {
diff --git a/Engine/Plugins/Enterprise/StaticMeshEditorExtension/Source/StaticMeshEditorExtension/Private/UVTools/UVGenerationUtils.cpp b/Engine/Plugins/Enterprise/StaticMeshEditorExtension/Source/StaticMeshEditorExtension/Private/UVTools/UVGenerationUtils.cpp
--- a/Engine/Plugins/Enterprise/StaticMeshEditorExtension/Source/StaticMeshEditorExtension/Private/UVTools/UVGenerationUtils.cpp
+++ b/Engine/Plugins/Enterprise/StaticMeshEditorExtension/Source/StaticMeshEditorExtension/Private/UVTools/UVGenerationUtils.cpp
@@ -1,60 +1,59 @@
 // Copyright 1998-2019 Epic Games, Inc. All Rights Reserved.
 
-#include "UVGenerationUtils.h"
-#include "PropertyHandle.h"
-#include "IDetailChildrenBuilder.h"
-#include "Widgets/Input/SVectorInputBox.h"
-#include "IDetailPropertyRow.h"
-#include "DetailWidgetRow.h"
-#include "Math/Axis.h"
+#include "UVTools/UVGenerationUtils.h"
+
+#include "Engine/StaticMesh.h"
 #include "MeshDescription.h"
+#include "MeshDescriptionOperations.h"
+#include "MeshUtilitiesCommon.h"
+#include "OverlappingCorners.h"
+#include "StaticMeshAttributes.h"
 
-TSharedRef<IPropertyTypeCustomization> FUVGenerationSettingsCustomization::MakeInstance()
-{
-	return MakeShareable(new FUVGenerationSettingsCustomization);
-}
 
-void FUVGenerationSettingsCustomization::CustomizeHeader(TSharedRef<IPropertyHandle> PropertyHandle, FDetailWidgetRow& HeaderRow, IPropertyTypeCustomizationUtils& CustomizationUtils)
+int32 UVGenerationUtils::GetNextOpenUVChannel(UStaticMesh* StaticMesh, int32 LODIndex)
 {
-	TArray<void*> StructPtrs;
-	PropertyHandle->AccessRawData(StructPtrs);
-	GenerateUVSettings = (StructPtrs.Num() == 1) ? reinterpret_cast<FUVGenerationSettings*>(StructPtrs[0]) : nullptr;
-}
+	FMeshDescription* Mesh = StaticMesh->GetMeshDescription(LODIndex);
+	int32 NumberOfUVs = StaticMesh->GetNumUVChannels(LODIndex);
+	int32 FirstEmptyUVs = 0;
 
-void FUVGenerationSettingsCustomization::CustomizeChildren(TSharedRef<IPropertyHandle> StructPropertyHandle, IDetailChildrenBuilder& ChildBuilder, IPropertyTypeCustomizationUtils& StructCustomizationUtils)
-{
-	uint32 NumChildProps = 0;
-	
-	StructPropertyHandle->GetNumChildren(NumChildProps);
-
-	for (uint32 Idx = 0; Idx < NumChildProps; Idx++)
+	for (; FirstEmptyUVs < NumberOfUVs; ++FirstEmptyUVs)
 	{
-		TSharedPtr<IPropertyHandle> PropHandle = StructPropertyHandle->GetChildHandle(Idx);
+		const TVertexInstanceAttributesConstRef<FVector2D> UVChannels = Mesh->VertexInstanceAttributes().GetAttributesRef<FVector2D>(MeshAttribute::VertexInstance::TextureCoordinate);
+		const FVector2D DefValue = UVChannels.GetDefaultValue();
+		bool bHasNonDefaultValue = false;
 
-		FName PropertyName = PropHandle->GetProperty()->GetFName();
-		if (PropertyName == GET_MEMBER_NAME_CHECKED(FUVGenerationSettings, Size)
-			|| PropertyName == GET_MEMBER_NAME_CHECKED(FUVGenerationSettings, Position)
-			|| PropertyName == GET_MEMBER_NAME_CHECKED(FUVGenerationSettings, Rotation))
-		{
-			PropHandle->SetOnChildPropertyValueChanged(FSimpleDelegate::CreateSP(this, &FUVGenerationSettingsCustomization::OnShapePropertyChanged));
-		}
-		else if (PropertyName == GET_MEMBER_NAME_CHECKED(FUVGenerationSettings, TargetChannel))
+		for (FVertexInstanceID InstanceID : Mesh->VertexInstances().GetElementIDs())
 		{
-			if (GenerateUVSettings && GenerateUVSettings->OnGetNumberOfUVs.IsBound())
+			if (UVChannels.Get(InstanceID, FirstEmptyUVs) != DefValue)
 			{
-				PropHandle->SetInstanceMetaData(TEXT("ClampMax"), FString::FromInt(FMath::Min(GenerateUVSettings->OnGetNumberOfUVs.Execute(), (int32)MAX_MESH_TEXTURE_COORDS_MD - 1)));
-			}
-			else
-			{
-				PropHandle->SetInstanceMetaData(TEXT("ClampMax"), FString::FromInt(MAX_MESH_TEXTURE_COORDS_MD - 1));
+				bHasNonDefaultValue = true;
+				break;
 			}
 		}
 
-		IDetailPropertyRow& PropertyRow = ChildBuilder.AddProperty(PropHandle.ToSharedRef());
+		if (!bHasNonDefaultValue)
+		{
+			//We found an "empty" channel.
+			break;
+		}
 	}
-}
+	return FirstEmptyUVs < MAX_MESH_TEXTURE_COORDS_MD ? FirstEmptyUVs : -1;
+};
 
-void FUVGenerationSettingsCustomization::OnShapePropertyChanged()
+void UVGenerationUtils::SetupGeneratedLightmapUVResolution(UStaticMesh* StaticMesh, int32 LODIndex)
 {
-	GenerateUVSettings->OnShapeEditingValueChanged.Broadcast();
-}
\ No newline at end of file
+	FMeshBuildSettings& BuildSettings = StaticMesh->GetSourceModel(LODIndex).BuildSettings;
+	FMeshDescription& Mesh = *StaticMesh->GetMeshDescription(LODIndex);
+
+	// Determine the absolute minimum lightmap resolution that can be used for packing
+	float ComparisonThreshold = BuildSettings.bRemoveDegenerates ? THRESH_POINTS_ARE_SAME : 0.0f;
+	FOverlappingCorners OverlappingCorners;
+	FMeshDescriptionOperations::FindOverlappingCorners(OverlappingCorners, Mesh, ComparisonThreshold);
+
+	// Packing expects at least one texel per chart. This is the absolute minimum to generate valid UVs.
+	int32 ChartCount = FMeshDescriptionOperations::GetUVChartCount(Mesh, BuildSettings.SrcLightmapIndex, ELightmapUVVersion::Latest, OverlappingCorners);
+	const int32 AbsoluteMinResolution = 1 << FMath::CeilLogTwo(FMath::Sqrt(ChartCount));
+	const int32 LightmapResolution = FMath::Clamp(BuildSettings.MinLightmapResolution, AbsoluteMinResolution, 512);
+
+	BuildSettings.MinLightmapResolution = LightmapResolution;
+}
diff --git a/Engine/Plugins/Enterprise/StaticMeshEditorExtension/Source/StaticMeshEditorExtension/Private/UVTools/UVGenerationUtils.h b/Engine/Plugins/Enterprise/StaticMeshEditorExtension/Source/StaticMeshEditorExtension/Private/UVTools/UVGenerationUtils.h
deleted file mode 100644
diff --git a/Engine/Plugins/Enterprise/StaticMeshEditorExtension/Source/StaticMeshEditorExtension/Public/UVTools/UVGenerationFlattenMapping.h b/Engine/Plugins/Enterprise/StaticMeshEditorExtension/Source/StaticMeshEditorExtension/Public/UVTools/UVGenerationFlattenMapping.h
--- a/Engine/Plugins/Enterprise/StaticMeshEditorExtension/Source/StaticMeshEditorExtension/Public/UVTools/UVGenerationFlattenMapping.h
+++ b/Engine/Plugins/Enterprise/StaticMeshEditorExtension/Source/StaticMeshEditorExtension/Public/UVTools/UVGenerationFlattenMapping.h
@@ -16,10 +16,10 @@ class STATICMESHEDITOREXTENSION_API UUVGenerationFlattenMapping : public UBluepr
 	GENERATED_BODY()
 
 public:
-	UFUNCTION(BlueprintCallable, Category = "Editor Scripting | StaticMesh", meta = (DisplayName = "Generate Flatten Mapping UVs"))
-	static void GenerateFlattenMappingUVs(UStaticMesh* InStaticMesh, int32 UVChannel, float AngleThreshold = 66.f, float AreaWeight = 0.7f);
+	UFUNCTION(BlueprintCallable, Category = "Editor Scripting | StaticMesh", meta = (DisplayName = "Generate Unwrapped UVs"))
+	static void GenerateFlattenMappingUVs(UStaticMesh* InStaticMesh, int32 UVChannel, float AngleThreshold = 66.f);
 
-	static void GenerateUVs(FMeshDescription& InMesh, int32 UVChannel, bool bRemoveDegenerate, float AngleThreshold = 66.f, float AreaWeight = 0.7f);
+	static void GenerateUVs(FMeshDescription& InMesh, int32 UVChannel, bool bRemoveDegenerate, float AngleThreshold = 66.f);
 
 private:
 
diff --git a/Engine/Plugins/Enterprise/StaticMeshEditorExtension/Source/StaticMeshEditorExtension/Public/UVTools/UVGenerationUtils.h b/Engine/Plugins/Enterprise/StaticMeshEditorExtension/Source/StaticMeshEditorExtension/Public/UVTools/UVGenerationUtils.h
new file mode 100644
--- /dev/null
+++ b/Engine/Plugins/Enterprise/StaticMeshEditorExtension/Source/StaticMeshEditorExtension/Public/UVTools/UVGenerationUtils.h
@@ -0,0 +1,25 @@
+// Copyright 1998-2019 Epic Games, Inc. All Rights Reserved.
+
+#pragma once
+
+#include "CoreMinimal.h"
+
+class UStaticMesh;
+
+class STATICMESHEDITOREXTENSION_API UVGenerationUtils
+{
+public:
+	/**
+	 * Returns the index of the next open UV channel, existing channels filled with default values are considered open.
+	 * If no open channel is found, returns -1. 
+	 */
+	static int32 GetNextOpenUVChannel(UStaticMesh* StaticMesh, int32 LODIndex);
+
+	/*
+	 * Setup the a good lightmap minimal resolution for the given static mesh LOD in its build settings.
+	 *
+	 * @param StaticMesh	The StaticMesh being set up
+	 * @param LODIndex		The LOD being set up.
+	 */
+	static void SetupGeneratedLightmapUVResolution(UStaticMesh* StaticMesh, int32 LODIndex);
+};
\ No newline at end of file
