Fixed crash in StaticMeshEditor when undoing mesh operations that triggered a PostEditChange on the StaticMesh.

diff --git a/Engine/Plugins/Enterprise/StaticMeshEditorExtension/Source/StaticMeshEditorExtension/Private/MeshEditingContext.cpp b/Engine/Plugins/Enterprise/StaticMeshEditorExtension/Source/StaticMeshEditorExtension/Private/MeshEditingContext.cpp
--- a/Engine/Plugins/Enterprise/StaticMeshEditorExtension/Source/StaticMeshEditorExtension/Private/MeshEditingContext.cpp
+++ b/Engine/Plugins/Enterprise/StaticMeshEditorExtension/Source/StaticMeshEditorExtension/Private/MeshEditingContext.cpp
@@ -35,6 +35,128 @@
 namespace MeshEditingContext
 {
 	static FAutoConsoleVariable OverlayHoverDistance(TEXT("MeshEditing.OverlayHoverDistance"), 0.01f, TEXT("Distance of overlay triangle to hover above element"));
+
+	void SetEditableMeshDescription(UEditableMesh* EditableMesh, UStaticMeshComponent* StaticMeshComponent, int32 LODIndex)
+	{
+		if (EditableMesh == nullptr || StaticMeshComponent == nullptr || LODIndex < 0)
+		{
+			return;
+		}
+
+		// Point EditableMesh's MeshDescription to StaticMesh's
+		FMeshDescription* EditableMeshDescription = EditableMesh->GetMeshDescription();
+		FMeshDescription* MeshDescription = StaticMeshComponent->GetStaticMesh()->GetMeshDescription(LODIndex);
+		if (MeshDescription == EditableMeshDescription)
+		{
+			return;
+		}
+
+		UStaticMesh* StaticMesh = StaticMeshComponent->GetStaticMesh();
+
+		// Register additional attributes required by features modifying EditableMesh
+		MeshDescription->PolygonAttributes().RegisterAttribute<FVector>( MeshAttribute::Polygon::Normal, 1, FVector::ZeroVector, EMeshAttributeFlags::Transient );
+		MeshDescription->PolygonAttributes().RegisterAttribute<FVector>( MeshAttribute::Polygon::Tangent, 1, FVector::ZeroVector, EMeshAttributeFlags::Transient );
+		MeshDescription->PolygonAttributes().RegisterAttribute<FVector>( MeshAttribute::Polygon::Binormal, 1, FVector::ZeroVector, EMeshAttributeFlags::Transient );
+		MeshDescription->PolygonAttributes().RegisterAttribute<FVector>( MeshAttribute::Polygon::Center, 1, FVector::ZeroVector, EMeshAttributeFlags::Transient );
+		MeshDescription->PolygonGroupAttributes().RegisterAttribute<FName>( MeshAttribute::PolygonGroup::MaterialAssetName );
+		MeshDescription->PolygonGroupAttributes().RegisterAttribute<bool>( MeshAttribute::PolygonGroup::EnableCollision );
+		MeshDescription->PolygonGroupAttributes().RegisterAttribute<bool>( MeshAttribute::PolygonGroup::CastShadow );
+
+		// Match EditableMesh material asset name with material slot names
+		TPolygonGroupAttributesConstRef<FName> SlotNames = MeshDescription->PolygonGroupAttributes().GetAttributesRef<FName>( MeshAttribute::PolygonGroup::ImportedMaterialSlotName );
+		TPolygonGroupAttributesRef<FName> AssetNames = MeshDescription->PolygonGroupAttributes().GetAttributesRef<FName>( MeshAttribute::PolygonGroup::MaterialAssetName );
+		TPolygonGroupAttributesRef<bool> EnableCollisions = MeshDescription->PolygonGroupAttributes().GetAttributesRef<bool>( MeshAttribute::PolygonGroup::EnableCollision );
+		TPolygonGroupAttributesRef<bool> CastShadows = MeshDescription->PolygonGroupAttributes().GetAttributesRef<bool>( MeshAttribute::PolygonGroup::CastShadow );
+
+		TPolygonGroupAttributesConstRef<FName> SrcSlotNames = EditableMeshDescription->PolygonGroupAttributes().GetAttributesRef<FName>( MeshAttribute::PolygonGroup::ImportedMaterialSlotName );
+		TPolygonGroupAttributesConstRef<FName> SrcAssetNames = EditableMeshDescription->PolygonGroupAttributes().GetAttributesRef<FName>( MeshAttribute::PolygonGroup::MaterialAssetName );
+		TPolygonGroupAttributesConstRef<bool> SrcEnableCollisions = EditableMeshDescription->PolygonGroupAttributes().GetAttributesRef<bool>( MeshAttribute::PolygonGroup::EnableCollision );
+		TPolygonGroupAttributesConstRef<bool> SrcCastShadows = EditableMeshDescription->PolygonGroupAttributes().GetAttributesRef<bool>( MeshAttribute::PolygonGroup::CastShadow );
+
+		for( const FPolygonGroupID PolygonGroupID : MeshDescription->PolygonGroups().GetElementIDs() )
+		{
+			FPolygonGroupID MatchingPolygonGroup(FPolygonGroupID::Invalid);
+			for( const FPolygonGroupID SrcPolygonGroupID : EditableMeshDescription->PolygonGroups().GetElementIDs() )
+			{
+				if( SlotNames[ PolygonGroupID ] == SrcSlotNames[ SrcPolygonGroupID ] )
+				{
+					MatchingPolygonGroup = SrcPolygonGroupID;
+					break;
+				}
+			}
+
+			if( MatchingPolygonGroup != FPolygonGroupID::Invalid )
+			{
+				AssetNames[ PolygonGroupID ] = SrcAssetNames[ MatchingPolygonGroup ];
+				EnableCollisions[ PolygonGroupID ] = SrcEnableCollisions[ MatchingPolygonGroup ];
+				CastShadows[ PolygonGroupID ] = SrcCastShadows[ MatchingPolygonGroup ];
+			}
+		}
+
+		{
+			// Compute the polygon attributes that are used by the StaticMeshEditor
+			// Taken from UEditableMesh::GeneratePolygonTangentsAndNormals
+			// TODO: Unify the code with FMeshDescriptionOperations::ConvertFromRawMesh 
+			TVertexAttributesConstRef<FVector> VertexPositions = MeshDescription->VertexAttributes().GetAttributesRef<FVector>(MeshAttribute::Vertex::Position);
+
+			TVertexInstanceAttributesConstRef<FVector2D> VertexUVs = MeshDescription->VertexInstanceAttributes().GetAttributesRef<FVector2D>(MeshAttribute::VertexInstance::TextureCoordinate);
+
+			TPolygonAttributesRef<FVector> PolygonNormals = MeshDescription->PolygonAttributes().GetAttributesRef<FVector>(MeshAttribute::Polygon::Normal);
+			TPolygonAttributesRef<FVector> PolygonTangents = MeshDescription->PolygonAttributes().GetAttributesRef<FVector>(MeshAttribute::Polygon::Tangent);
+			TPolygonAttributesRef<FVector> PolygonBinormals = MeshDescription->PolygonAttributes().GetAttributesRef<FVector>(MeshAttribute::Polygon::Binormal);
+			TPolygonAttributesRef<FVector> PolygonCenters = MeshDescription->PolygonAttributes().GetAttributesRef<FVector>(MeshAttribute::Polygon::Center);
+
+			for (const FPolygonID PolygonID : MeshDescription->Polygons().GetElementIDs())
+			{
+				// Calculate the center of this polygon
+				FVector Center = FVector::ZeroVector;
+				const TArray<FVertexInstanceID>& VertexInstanceIDs = MeshDescription->GetPolygonVertexInstances(PolygonID);
+				for (const FVertexInstanceID VertexInstanceID : VertexInstanceIDs)
+				{
+					Center += VertexPositions[MeshDescription->GetVertexInstanceVertex(VertexInstanceID)];
+				}
+				Center /= float(VertexInstanceIDs.Num());
+
+				// Calculate the tangent basis for the polygon, based on the average of all constituent triangles
+				FVector Normal = FVector::ZeroVector;
+				FVector Tangent = FVector::ZeroVector;
+				FVector Binormal = FVector::ZeroVector;
+
+				for (const FTriangleID TriangleID : MeshDescription->GetPolygonTriangleIDs(PolygonID))
+				{
+					TArrayView<const FVertexInstanceID> TriVertexInstanceIDs = MeshDescription->GetTriangleVertexInstances(TriangleID);
+					const FVertexID VertexID0 = MeshDescription->GetVertexInstanceVertex(TriVertexInstanceIDs[0]);
+					const FVertexID VertexID1 = MeshDescription->GetVertexInstanceVertex(TriVertexInstanceIDs[1]);
+					const FVertexID VertexID2 = MeshDescription->GetVertexInstanceVertex(TriVertexInstanceIDs[2]);
+
+					const FVector DPosition1 = VertexPositions[VertexID1] - VertexPositions[VertexID0];
+					const FVector DPosition2 = VertexPositions[VertexID2] - VertexPositions[VertexID0];
+
+					const FVector2D DUV1 = VertexUVs.Get(TriVertexInstanceIDs[1], 0) - VertexUVs.Get(TriVertexInstanceIDs[0], 0);
+					const FVector2D DUV2 = VertexUVs.Get(TriVertexInstanceIDs[2], 0) - VertexUVs.Get(TriVertexInstanceIDs[0], 0);
+
+					// We have a left-handed coordinate system, but a counter-clockwise winding order
+					// Hence normal calculation has to take the triangle vectors cross product in reverse.
+					Normal += FVector::CrossProduct(DPosition2, DPosition1);
+
+					// ...and tangent space seems to be right-handed.
+					const float DetUV = FVector2D::CrossProduct(DUV1, DUV2);
+					const float InvDetUV = (DetUV == 0.0f) ? 0.0f : 1.0f / DetUV;
+
+					Tangent += (DPosition1 * DUV2.Y - DPosition2 * DUV1.Y) * InvDetUV;
+					Binormal += (DPosition2 * DUV1.X - DPosition1 * DUV2.X) * InvDetUV;
+				}
+
+				PolygonNormals[PolygonID] = Normal.GetSafeNormal();
+				PolygonTangents[PolygonID] = Tangent.GetSafeNormal();
+				PolygonBinormals[PolygonID] = Binormal.GetSafeNormal();
+				PolygonCenters[PolygonID] = Center;
+			}
+		}
+
+		EditableMesh->SetMeshDescription(MeshDescription);
+		EditableMesh->InitializeAdapters();
+	}
 }
 
 
@@ -99,5 +221,10 @@
 
 const UEditableMesh* FEditableMeshCache::FindEditableMesh(const UPrimitiveComponent& Component, const FEditableMeshSubMeshAddress& SubMeshAddress) const
 {
+	return FindModifiableEditableMesh(Component, SubMeshAddress);
+}
+
+UEditableMesh* FEditableMeshCache::FindModifiableEditableMesh(const UPrimitiveComponent& Component, const FEditableMeshSubMeshAddress& SubMeshAddress) const
+{
 	if (Cast<UStaticMeshComponent>(&Component) == nullptr || SubMeshAddress.EditableMeshFormat == nullptr)
 	{
@@ -132,6 +259,24 @@
 	}
 }
 
+void FEditableMeshCache::ResetObject(UStaticMesh * StaticMesh)
+{
+	const UStaticMeshComponent* StaticMeshComponent = StaticMeshesToComponents.FindRef(StaticMesh);
+	if (StaticMeshComponent)
+	{
+		// If one of the LODs has been edited, reset its MeshDescription
+		for (int32 LODIndex = 0; LODIndex < StaticMesh->GetNumSourceModels(); ++LODIndex)
+		{
+			const FEditableMeshSubMeshAddress SubMeshAddressToQuery = UEditableMeshFactory::MakeSubmeshAddress(const_cast<UStaticMeshComponent*>(StaticMeshComponent), LODIndex);
+			TStrongObjectPtr<UEditableMesh> EditableMeshPtr = CachedEditableMeshes.FindRef(SubMeshAddressToQuery);
+			if (EditableMeshPtr)
+			{
+				EditableMeshPtr->SetMeshDescription(&EditableMeshPtr->OwnedMeshDescription);
+			}
+		}
+	}
+}
+
 FMeshEditingContext::FMeshEditingContext()
 	: LODIndex(INDEX_NONE)
 	, StaticMeshComponent(nullptr)
@@ -210,111 +355,7 @@
 	FMeshDescription* MeshDescription = StaticMeshComponent->GetStaticMesh()->GetMeshDescription(LODIndex);
 	if (MeshDescription != EditableMeshDescription)
 	{
-		UStaticMesh* StaticMesh = StaticMeshComponent->GetStaticMesh();
+		MeshEditingContext::SetEditableMeshDescription(EditableMesh, StaticMeshComponent, LODIndex);
 
-		// Register additional attributes required by features modifying EditableMesh
-		MeshDescription->PolygonAttributes().RegisterAttribute<FVector>( MeshAttribute::Polygon::Normal, 1, FVector::ZeroVector, EMeshAttributeFlags::Transient );
-		MeshDescription->PolygonAttributes().RegisterAttribute<FVector>( MeshAttribute::Polygon::Tangent, 1, FVector::ZeroVector, EMeshAttributeFlags::Transient );
-		MeshDescription->PolygonAttributes().RegisterAttribute<FVector>( MeshAttribute::Polygon::Binormal, 1, FVector::ZeroVector, EMeshAttributeFlags::Transient );
-		MeshDescription->PolygonAttributes().RegisterAttribute<FVector>( MeshAttribute::Polygon::Center, 1, FVector::ZeroVector, EMeshAttributeFlags::Transient );
-		MeshDescription->PolygonGroupAttributes().RegisterAttribute<FName>( MeshAttribute::PolygonGroup::MaterialAssetName );
-		MeshDescription->PolygonGroupAttributes().RegisterAttribute<bool>( MeshAttribute::PolygonGroup::EnableCollision );
-		MeshDescription->PolygonGroupAttributes().RegisterAttribute<bool>( MeshAttribute::PolygonGroup::CastShadow );
-
-		// Match EditableMesh material asset name with material slot names
-		TPolygonGroupAttributesConstRef<FName> SlotNames = MeshDescription->PolygonGroupAttributes().GetAttributesRef<FName>( MeshAttribute::PolygonGroup::ImportedMaterialSlotName );
-		TPolygonGroupAttributesRef<FName> AssetNames = MeshDescription->PolygonGroupAttributes().GetAttributesRef<FName>( MeshAttribute::PolygonGroup::MaterialAssetName );
-		TPolygonGroupAttributesRef<bool> EnableCollisions = MeshDescription->PolygonGroupAttributes().GetAttributesRef<bool>( MeshAttribute::PolygonGroup::EnableCollision );
-		TPolygonGroupAttributesRef<bool> CastShadows = MeshDescription->PolygonGroupAttributes().GetAttributesRef<bool>( MeshAttribute::PolygonGroup::CastShadow );
-
-		TPolygonGroupAttributesConstRef<FName> SrcSlotNames = EditableMeshDescription->PolygonGroupAttributes().GetAttributesRef<FName>( MeshAttribute::PolygonGroup::ImportedMaterialSlotName );
-		TPolygonGroupAttributesConstRef<FName> SrcAssetNames = EditableMeshDescription->PolygonGroupAttributes().GetAttributesRef<FName>( MeshAttribute::PolygonGroup::MaterialAssetName );
-		TPolygonGroupAttributesConstRef<bool> SrcEnableCollisions = EditableMeshDescription->PolygonGroupAttributes().GetAttributesRef<bool>( MeshAttribute::PolygonGroup::EnableCollision );
-		TPolygonGroupAttributesConstRef<bool> SrcCastShadows = EditableMeshDescription->PolygonGroupAttributes().GetAttributesRef<bool>( MeshAttribute::PolygonGroup::CastShadow );
-
-		for( const FPolygonGroupID PolygonGroupID : MeshDescription->PolygonGroups().GetElementIDs() )
-		{
-			FPolygonGroupID MatchingPolygonGroup(FPolygonGroupID::Invalid);
-			for( const FPolygonGroupID SrcPolygonGroupID : EditableMeshDescription->PolygonGroups().GetElementIDs() )
-			{
-				if( SlotNames[ PolygonGroupID ] == SrcSlotNames[ SrcPolygonGroupID ] )
-				{
-					MatchingPolygonGroup = SrcPolygonGroupID;
-					break;
-				}
-			}
-
-			if( MatchingPolygonGroup != FPolygonGroupID::Invalid )
-			{
-				AssetNames[ PolygonGroupID ] = SrcAssetNames[ MatchingPolygonGroup ];
-				EnableCollisions[ PolygonGroupID ] = SrcEnableCollisions[ MatchingPolygonGroup ];
-				CastShadows[ PolygonGroupID ] = SrcCastShadows[ MatchingPolygonGroup ];
-			}
-		}
-
-		{
-			// Compute the polygon attributes that are used by the StaticMeshEditor
-			// Taken from UEditableMesh::GeneratePolygonTangentsAndNormals
-			// TODO: Unify the code with FMeshDescriptionOperations::ConvertFromRawMesh 
-			TVertexAttributesConstRef<FVector> VertexPositions = MeshDescription->VertexAttributes().GetAttributesRef<FVector>(MeshAttribute::Vertex::Position);
-
-			TVertexInstanceAttributesConstRef<FVector2D> VertexUVs = MeshDescription->VertexInstanceAttributes().GetAttributesRef<FVector2D>(MeshAttribute::VertexInstance::TextureCoordinate);
-
-			TPolygonAttributesRef<FVector> PolygonNormals = MeshDescription->PolygonAttributes().GetAttributesRef<FVector>(MeshAttribute::Polygon::Normal);
-			TPolygonAttributesRef<FVector> PolygonTangents = MeshDescription->PolygonAttributes().GetAttributesRef<FVector>(MeshAttribute::Polygon::Tangent);
-			TPolygonAttributesRef<FVector> PolygonBinormals = MeshDescription->PolygonAttributes().GetAttributesRef<FVector>(MeshAttribute::Polygon::Binormal);
-			TPolygonAttributesRef<FVector> PolygonCenters = MeshDescription->PolygonAttributes().GetAttributesRef<FVector>(MeshAttribute::Polygon::Center);
-
-			for (const FPolygonID PolygonID : MeshDescription->Polygons().GetElementIDs())
-			{
-				// Calculate the center of this polygon
-				FVector Center = FVector::ZeroVector;
-				const TArray<FVertexInstanceID>& VertexInstanceIDs = MeshDescription->GetPolygonVertexInstances(PolygonID);
-				for (const FVertexInstanceID VertexInstanceID : VertexInstanceIDs)
-				{
-					Center += VertexPositions[MeshDescription->GetVertexInstanceVertex(VertexInstanceID)];
-				}
-				Center /= float(VertexInstanceIDs.Num());
-
-				// Calculate the tangent basis for the polygon, based on the average of all constituent triangles
-				FVector Normal = FVector::ZeroVector;
-				FVector Tangent = FVector::ZeroVector;
-				FVector Binormal = FVector::ZeroVector;
-
-				for (const FTriangleID TriangleID : MeshDescription->GetPolygonTriangleIDs(PolygonID))
-				{
-					TArrayView<const FVertexInstanceID> TriVertexInstanceIDs = MeshDescription->GetTriangleVertexInstances(TriangleID);
-					const FVertexID VertexID0 = MeshDescription->GetVertexInstanceVertex(TriVertexInstanceIDs[0]);
-					const FVertexID VertexID1 = MeshDescription->GetVertexInstanceVertex(TriVertexInstanceIDs[1]);
-					const FVertexID VertexID2 = MeshDescription->GetVertexInstanceVertex(TriVertexInstanceIDs[2]);
-
-					const FVector DPosition1 = VertexPositions[VertexID1] - VertexPositions[VertexID0];
-					const FVector DPosition2 = VertexPositions[VertexID2] - VertexPositions[VertexID0];
-
-					const FVector2D DUV1 = VertexUVs.Get(TriVertexInstanceIDs[1], 0) - VertexUVs.Get(TriVertexInstanceIDs[0], 0);
-					const FVector2D DUV2 = VertexUVs.Get(TriVertexInstanceIDs[2], 0) - VertexUVs.Get(TriVertexInstanceIDs[0], 0);
-
-					// We have a left-handed coordinate system, but a counter-clockwise winding order
-					// Hence normal calculation has to take the triangle vectors cross product in reverse.
-					Normal += FVector::CrossProduct(DPosition2, DPosition1);
-
-					// ...and tangent space seems to be right-handed.
-					const float DetUV = FVector2D::CrossProduct(DUV1, DUV2);
-					const float InvDetUV = (DetUV == 0.0f) ? 0.0f : 1.0f / DetUV;
-
-					Tangent += (DPosition1 * DUV2.Y - DPosition2 * DUV1.Y) * InvDetUV;
-					Binormal += (DPosition2 * DUV1.X - DPosition1 * DUV2.X) * InvDetUV;
-				}
-
-				PolygonNormals[PolygonID] = Normal.GetSafeNormal();
-				PolygonTangents[PolygonID] = Tangent.GetSafeNormal();
-				PolygonBinormals[PolygonID] = Binormal.GetSafeNormal();
-				PolygonCenters[PolygonID] = Center;
-			}
-		}
-
-		EditableMesh->SetMeshDescription(MeshDescription);
-		EditableMesh->InitializeAdapters();
-
 		// Disable spatial database to flush it
 		EditableMesh->SetAllowSpatialDatabase( false );
@@ -691,9 +732,9 @@
 		return;
 	}
 
-	// A PostEdit has been called on the edited static mesh, remove associated editable mesh
-	// It is not valid anymore
-	FEditableMeshCache::Get().RemoveObject(StaticMeshComponent->GetStaticMesh());
+	// A PostEdit has been called on the edited static mesh, the associated editable mesh has to be reset
+	// The same EditableMesh has to be reused for the undo operations since they reference it
+	FEditableMeshCache::Get().ResetObject(StaticMeshComponent->GetStaticMesh());
 
 	// Re-initialize EditableMesh if context was active
 	if (LODIndex != INDEX_NONE)
@@ -706,6 +747,21 @@
 		// Regenerate editable mesh and related data
 		SetLODIndex(CachedLODIndex);
 	}
+
+	// Make sure the EditableMesh MeshDescription for all LODs are updated after the ResetObject (even when not in Edit Mode, LODIndex == INDEX_NONE)
+	// The user could be undoing operations outside of Edit Mode
+	for (int32 CurrentLODIndex = 0; CurrentLODIndex < StaticMeshComponent->GetStaticMesh()->GetNumSourceModels(); ++CurrentLODIndex)
+	{
+		// MeshDescription for LODIndex is already set through SetLODIndex
+		if (CurrentLODIndex == LODIndex)
+		{
+			continue;
+		}
+
+		FEditableMeshSubMeshAddress SubMeshAddressToQuery = UEditableMeshFactory::MakeSubmeshAddress(StaticMeshComponent, CurrentLODIndex);
+		UEditableMesh* CurrentEditableMesh = FEditableMeshCache::Get().FindModifiableEditableMesh(*StaticMeshComponent, SubMeshAddressToQuery);
+		MeshEditingContext::SetEditableMeshDescription(CurrentEditableMesh, StaticMeshComponent, CurrentLODIndex);
+	}
 }
 
 void FMeshEditingUIContext::ClearHoveredElements()
diff --git a/Engine/Plugins/Enterprise/StaticMeshEditorExtension/Source/StaticMeshEditorExtension/Private/MeshEditingContext.h b/Engine/Plugins/Enterprise/StaticMeshEditorExtension/Source/StaticMeshEditorExtension/Private/MeshEditingContext.h
--- a/Engine/Plugins/Enterprise/StaticMeshEditorExtension/Source/StaticMeshEditorExtension/Private/MeshEditingContext.h
+++ b/Engine/Plugins/Enterprise/StaticMeshEditorExtension/Source/StaticMeshEditorExtension/Private/MeshEditingContext.h
@@ -32,10 +32,14 @@ public:
 
 	/** Returns editable mesh associated with component and sub-mesh address, null pointer if not in cache */
 	const UEditableMesh* FindEditableMesh( const UPrimitiveComponent& Component, const FEditableMeshSubMeshAddress& SubMeshAddress ) const;
+	UEditableMesh* FindModifiableEditableMesh( const UPrimitiveComponent& Component, const FEditableMeshSubMeshAddress& SubMeshAddress ) const;
 
 	/** Removes editable meshes associated with static mesh from cache if applicable */
 	void RemoveObject(UStaticMesh* StaticMesh);
 
+	/** Resets the editable meshes associated with static mesh from cache if applicable */
+	void ResetObject(UStaticMesh* StaticMesh);
+
 private:
 	FEditableMeshCache() {}
 
diff --git a/Engine/Plugins/Enterprise/StaticMeshEditorExtension/Source/StaticMeshEditorExtension/Private/PolygonEditingToolbar.cpp b/Engine/Plugins/Enterprise/StaticMeshEditorExtension/Source/StaticMeshEditorExtension/Private/PolygonEditingToolbar.cpp
--- a/Engine/Plugins/Enterprise/StaticMeshEditorExtension/Source/StaticMeshEditorExtension/Private/PolygonEditingToolbar.cpp
+++ b/Engine/Plugins/Enterprise/StaticMeshEditorExtension/Source/StaticMeshEditorExtension/Private/PolygonEditingToolbar.cpp
@@ -41,6 +41,8 @@
 
 #define LOCTEXT_NAMESPACE "StaticMeshEditorExtensionToolbar"
 
+DEFINE_LOG_CATEGORY_STATIC(LogStaticMeshEditorExtension, Log, All);
+
 namespace FPolygonEditingCommandsUtil
 {
 	template< typename InMeshEditorCommand >
@@ -190,6 +192,7 @@ FPolygonEditingToolbar::FPolygonEditingToolbar()
 	, StaticMesh(nullptr)
 	, PolygonSelectionTool(nullptr)
 	, bDeleteCommandOverriden(false)
+	, bClearUndoTransactions(false)
 {
 	PolygonToolbarProxyObject = TStrongObjectPtr<UPolygonToolbarProxyObject>(NewObject<UPolygonToolbarProxyObject>());
 	PolygonToolbarProxyObject->Owner = this;
@@ -218,6 +221,12 @@ FPolygonEditingToolbar::~FPolygonEditingToolbar()
 		FEditableMeshCache::Get().RemoveObject(StaticMesh);
 	}
 
+	if (bClearUndoTransactions)
+	{
+		GEditor->ResetTransaction(LOCTEXT("MeshEditingResetTransaction", "Destroying Mesh Editing toolbar"));
+		UE_LOG(LogStaticMeshEditorExtension, Warning, TEXT("Undo history cleared after closing Static Mesh Editor"));
+	}
+
 	bIsEditing = false;
 
 	// Delete editing context
@@ -625,7 +634,9 @@ void FPolygonEditingToolbar::OnMeshChanged()
 	if (bIsEditing)
 	{
 		// Check if current selected LOD is still editable
-		if (!EditableLODs.IsValidIndex(StaticMeshEditor->GetCurrentLODLevel()) || !EditableLODs[StaticMeshEditor->GetCurrentLODLevel()])
+		// Note that EditableLODIndex is shifted by 1 EditableLODs since index 0 in EditableLODs is the LOD Auto
+		int32 EditableLODIndex = StaticMeshEditor->GetCurrentLODIndex() + 1;
+		if (!EditableLODs.IsValidIndex(EditableLODIndex) || !EditableLODs[EditableLODIndex])
 		{
 			if (StaticMeshEditor->GetCurrentLODLevel() == 0 && !EditableLODs[0])
 			{
@@ -652,8 +663,8 @@ void FPolygonEditingToolbar::UpdateEditableLODs()
 	// Build list of valid LODs
 	EditableLODs.AddUninitialized(StaticMesh->GetNumSourceModels() + 1);
 
-	// 'LOD Auto' is not a valid selection
-	EditableLODs[0] = false;
+	// 'LOD Auto' is not a valid selection except in case of only 1 LOD, then it is equivalent to LOD 0
+	EditableLODs[0] = StaticMesh->GetNumSourceModels() <= 1;
 
 	// LOD 0 is assumed to be a valid selection
 	EditableLODs[1] = true;
@@ -677,6 +688,12 @@ void FPolygonEditingToolbar::TrackUndo(UObject* Object, TUniquePtr<FChange> Reve
 		check( GUndo != nullptr || GEditor == nullptr || GEditor->bIsSimulatingInEditor );
 		if( GUndo != nullptr )
 		{
+			UEditableMesh* EditableMesh = Cast<UEditableMesh>(Object);
+			if (EditableMesh)
+			{
+				// Make sure to clear undo buffer if there were any EditableMesh changes
+				bClearUndoTransactions = true;
+			}
 			GUndo->StoreUndo( Object, MoveTemp( RevertChange ) );
 		}
 	}
diff --git a/Engine/Plugins/Enterprise/StaticMeshEditorExtension/Source/StaticMeshEditorExtension/Private/PolygonEditingToolbar.h b/Engine/Plugins/Enterprise/StaticMeshEditorExtension/Source/StaticMeshEditorExtension/Private/PolygonEditingToolbar.h
--- a/Engine/Plugins/Enterprise/StaticMeshEditorExtension/Source/StaticMeshEditorExtension/Private/PolygonEditingToolbar.h
+++ b/Engine/Plugins/Enterprise/StaticMeshEditorExtension/Source/StaticMeshEditorExtension/Private/PolygonEditingToolbar.h
@@ -311,9 +311,12 @@ private:
 	/** Dummy array used to implement method of the IMeshEditorModeUIContract interface */
 	TArray<TTuple<TSharedPtr<FUICommandInfo>, FUIAction>> ActionArray;
 
-	/* Flag to toggle the toolbar-specific bindings */
+	/** Flag to toggle the toolbar-specific bindings */
 	mutable bool bDeleteCommandOverriden;
 
-	/* Handle on callback when edited static mesh is re-imported */
+	/** Handle on callback when edited static mesh is re-imported */
 	FDelegateHandle OnObjectReimportedHandle;
+
+	/** Flag to indicate if the undo transaction buffer has to be cleared */
+	bool bClearUndoTransactions;
 };
diff --git a/Engine/Plugins/Enterprise/StaticMeshEditorExtension/Source/StaticMeshEditorExtension/Private/PolygonSelectionTool.cpp b/Engine/Plugins/Enterprise/StaticMeshEditorExtension/Source/StaticMeshEditorExtension/Private/PolygonSelectionTool.cpp
--- a/Engine/Plugins/Enterprise/StaticMeshEditorExtension/Source/StaticMeshEditorExtension/Private/PolygonSelectionTool.cpp
+++ b/Engine/Plugins/Enterprise/StaticMeshEditorExtension/Source/StaticMeshEditorExtension/Private/PolygonSelectionTool.cpp
@@ -22,8 +22,6 @@
 
 #define LOCTEXT_NAMESPACE "StaticMeshEditorExtensionToolbar"
 
-DEFINE_LOG_CATEGORY_STATIC(LogStaticMeshEditorExtension, Log, All);
-
 const FEditorModeID FPolygonSelectionTool::EM_PolygonSelection(TEXT("EM_PolygonSelection"));
 
 // @todo mesheditor extensibility: This should probably be removed after we've evicted all current mesh editing actions to another module
