Fix merge error in SUVGenerationTool.cpp
#rnx

diff --git a/Engine/Plugins/Enterprise/StaticMeshEditorExtension/Source/StaticMeshEditorExtension/Private/UVTools/SUVGenerationTool.cpp b/Engine/Plugins/Enterprise/StaticMeshEditorExtension/Source/StaticMeshEditorExtension/Private/UVTools/SUVGenerationTool.cpp
--- a/Engine/Plugins/Enterprise/StaticMeshEditorExtension/Source/StaticMeshEditorExtension/Private/UVTools/SUVGenerationTool.cpp
+++ b/Engine/Plugins/Enterprise/StaticMeshEditorExtension/Source/StaticMeshEditorExtension/Private/UVTools/SUVGenerationTool.cpp
@@ -36,13 +36,6 @@ SGenerateUV::SGenerateUV()
 	GenerateUVSettings = &SettingObjectUIHolder->GenerateUVSettings;
 }
 
-SGenerateUV::SGenerateUV()
-{
-	SettingObjectUIHolder = NewObject<UGenerateUVSettingsUIHolder>();
-	SettingObjectUIHolder->AddToRoot();
-	GenerateUVSettings = &SettingObjectUIHolder->GenerateUVSettings;
-}
-
 SGenerateUV::~SGenerateUV()
 {
 	SetPreviewModeActivated(false);
