Integrating CL# 10809916 from Dev-Enterprise
Fix crash

diff --git a/Engine/Plugins/Enterprise/StaticMeshEditorExtension/Source/StaticMeshEditorExtension/Private/UVTools/UVGenerationUtils.cpp b/Engine/Plugins/Enterprise/StaticMeshEditorExtension/Source/StaticMeshEditorExtension/Private/UVTools/UVGenerationUtils.cpp
--- a/Engine/Plugins/Enterprise/StaticMeshEditorExtension/Source/StaticMeshEditorExtension/Private/UVTools/UVGenerationUtils.cpp
+++ b/Engine/Plugins/Enterprise/StaticMeshEditorExtension/Source/StaticMeshEditorExtension/Private/UVTools/UVGenerationUtils.cpp
@@ -12,6 +12,11 @@
 
 int32 UVGenerationUtils::GetNextOpenUVChannel(UStaticMesh* StaticMesh, int32 LODIndex)
 {
+	if (!StaticMesh->IsMeshDescriptionValid(LODIndex))
+	{
+		return -1;
+	}
+
 	FMeshDescription* Mesh = StaticMesh->GetMeshDescription(LODIndex);
 	int32 NumberOfUVs = StaticMesh->GetNumUVChannels(LODIndex);
 	int32 FirstEmptyUVs = 0;
@@ -42,6 +47,11 @@ int32 UVGenerationUtils::GetNextOpenUVChannel(UStaticMesh* StaticMesh, int32 LOD
 
 void UVGenerationUtils::SetupGeneratedLightmapUVResolution(UStaticMesh* StaticMesh, int32 LODIndex)
 {
+	if (!StaticMesh->IsMeshDescriptionValid(LODIndex))
+	{
+		return;
+	}
+
 	FMeshBuildSettings& BuildSettings = StaticMesh->GetSourceModel(LODIndex).BuildSettings;
 	FMeshDescription& Mesh = *StaticMesh->GetMeshDescription(LODIndex);
 
