Added logic to properly remove transaction referring to editable mesh

#rnx

diff --git a/Engine/Plugins/Enterprise/StaticMeshEditorExtension/Source/StaticMeshEditorExtension/Private/PolygonEditingToolbar.cpp b/Engine/Plugins/Enterprise/StaticMeshEditorExtension/Source/StaticMeshEditorExtension/Private/PolygonEditingToolbar.cpp
--- a/Engine/Plugins/Enterprise/StaticMeshEditorExtension/Source/StaticMeshEditorExtension/Private/PolygonEditingToolbar.cpp
+++ b/Engine/Plugins/Enterprise/StaticMeshEditorExtension/Source/StaticMeshEditorExtension/Private/PolygonEditingToolbar.cpp
@@ -78,6 +78,87 @@ namespace FPolygonEditingCommandsUtil
 	}
 }
 
+//
+// Helper class to manage changes made on the editable mesh associated with the polygon editing toolbar
+// The goal is to have FChange::HasExpired to return false when the editable mesh is not reachable anymore.
+class FMeshEditingChange : public FSwapChange
+{
+
+public:
+
+	/** Constructor */
+	explicit FMeshEditingChange(TSharedPtr<FPolygonEditingToolbar> InToolbar, TUniquePtr<FChange>&& InBaseChange)
+		: Toolbar(InToolbar)
+		, BaseChange(MoveTemp(InBaseChange))
+	{
+		ensure(Toolbar.IsValid());
+	}
+
+	// Parent class overrides
+	virtual EChangeStyle GetChangeType() override
+	{
+		return Toolbar.IsValid() && BaseChange.IsValid() ? BaseChange->GetChangeType() : EChangeStyle::InPlaceSwap;
+	}
+
+	virtual TUniquePtr<FChange> Execute(UObject* Object) override
+	{
+		if(Toolbar.IsValid() && BaseChange.IsValid())
+		{
+			TUniquePtr<FMeshEditingChange> ExecutedChange( new FMeshEditingChange(Toolbar.Pin(), BaseChange->Execute(Object)) );
+			return ExecutedChange;
+		}
+
+		return nullptr;
+	}
+
+	virtual void Apply( UObject* Object ) override
+	{
+		if(Toolbar.IsValid() && BaseChange.IsValid())
+		{
+			BaseChange->Apply(Object);
+		}
+	}
+
+	virtual void Revert( UObject* Object ) override
+	{
+		if(Toolbar.IsValid() && BaseChange.IsValid())
+		{
+			BaseChange->Revert(Object);
+		}
+	}
+
+	/**
+	 * @remark Will return false when the StaticMesh editor is closed and the associated editable mesh is released
+	 */
+	virtual bool HasExpired( UObject* Object ) const
+	{ 
+		return !Toolbar.IsValid();
+	}
+
+	virtual FString ToString() const override
+	{
+		return Toolbar.IsValid() && BaseChange.IsValid() ? BaseChange->ToString() : FString();
+	}
+
+	virtual void PrintToLog(class FFeedbackContext& FeedbackContext, const int32 IndentLevel = 0) override
+	{
+		if(Toolbar.IsValid() && BaseChange.IsValid())
+		{
+			BaseChange->PrintToLog(FeedbackContext, IndentLevel);
+		}
+	}
+
+private:
+	TWeakPtr<FPolygonEditingToolbar> Toolbar;
+	TUniquePtr<FChange> BaseChange;
+
+private:
+	// Non-copyable
+	FMeshEditingChange( const FMeshEditingChange& ) = delete;
+	FMeshEditingChange& operator=( const FMeshEditingChange& ) = delete;
+};
+
+
 #define IMAGE_PLUGIN_BRUSH( RelativePath, ... ) FSlateImageBrush( FPolygonEditingToolbarStyle::InContent( RelativePath, ".png" ), __VA_ARGS__ )
 
 class FPolygonEditingToolbarStyle
@@ -192,7 +273,7 @@ FPolygonEditingToolbar::FPolygonEditingToolbar()
 	, StaticMesh(nullptr)
 	, PolygonSelectionTool(nullptr)
 	, bDeleteCommandOverriden(false)
-	, bClearUndoTransactions(false)
+	, bTransactionsRecorded(false)
 {
 	PolygonToolbarProxyObject = TStrongObjectPtr<UPolygonToolbarProxyObject>(NewObject<UPolygonToolbarProxyObject>());
 	PolygonToolbarProxyObject->Owner = this;
@@ -219,12 +300,11 @@ FPolygonEditingToolbar::~FPolygonEditingToolbar()
 
 		// Remove editable meshes related to static mesh from cache
 		FEditableMeshCache::Get().RemoveObject(StaticMesh);
-	}
 
-	if (bClearUndoTransactions)
-	{
-		GEditor->ResetTransaction(LOCTEXT("MeshEditingResetTransaction", "Destroying Mesh Editing toolbar"));
-		UE_LOG(LogStaticMeshEditorExtension, Warning, TEXT("Undo history cleared after closing Static Mesh Editor"));
+		if(bTransactionsRecorded)
+		{
+			UE_LOG(LogStaticMeshEditorExtension, Warning, TEXT("Mesh editing operations made on static mesh %s have been nullified. Undoing those mesh editing operations won't have any effect."), *StaticMesh->GetName());
+		}
 	}
 
 	bIsEditing = false;
@@ -688,13 +768,13 @@ void FPolygonEditingToolbar::TrackUndo(UObject* Object, TUniquePtr<FChange> Reve
 		check( GUndo != nullptr || GEditor == nullptr || GEditor->bIsSimulatingInEditor );
 		if( GUndo != nullptr )
 		{
-			UEditableMesh* EditableMesh = Cast<UEditableMesh>(Object);
-			if (EditableMesh)
+			if (UEditableMesh* EditableMesh = Cast<UEditableMesh>(Object))
 			{
-				// Make sure to clear undo buffer if there were any EditableMesh changes
-				bClearUndoTransactions = true;
+				// Create custom FChange object and add it to current transaction
+				TUniquePtr<FMeshEditingChange> Change( new FMeshEditingChange( AsShared(), MoveTemp( RevertChange ) ) );
+				GUndo->StoreUndo( Object, MoveTemp( Change ) );
+				bTransactionsRecorded = true;
 			}
-			GUndo->StoreUndo( Object, MoveTemp( RevertChange ) );
 		}
 	}
 }
diff --git a/Engine/Plugins/Enterprise/StaticMeshEditorExtension/Source/StaticMeshEditorExtension/Private/PolygonEditingToolbar.h b/Engine/Plugins/Enterprise/StaticMeshEditorExtension/Source/StaticMeshEditorExtension/Private/PolygonEditingToolbar.h
--- a/Engine/Plugins/Enterprise/StaticMeshEditorExtension/Source/StaticMeshEditorExtension/Private/PolygonEditingToolbar.h
+++ b/Engine/Plugins/Enterprise/StaticMeshEditorExtension/Source/StaticMeshEditorExtension/Private/PolygonEditingToolbar.h
@@ -317,6 +317,9 @@ private:
 	/** Handle on callback when edited static mesh is re-imported */
 	FDelegateHandle OnObjectReimportedHandle;
 
-	/** Flag to indicate if the undo transaction buffer has to be cleared */
-	bool bClearUndoTransactions;
+	/**
+	 * True if mesh editing operations have been executed on the static mesh
+	 * If true when the static mesh editor closes, a warning message will be logged about loss of mesh editing operations
+	 */
+	bool bTransactionsRecorded;
 };
