Resolved some of the conflicts from Main

#rnx

diff --git a/Engine/Plugins/Enterprise/StaticMeshEditorExtension/Source/StaticMeshEditorExtension/Private/UVTools/SUVGenerationTool.cpp b/Engine/Plugins/Enterprise/StaticMeshEditorExtension/Source/StaticMeshEditorExtension/Private/UVTools/SUVGenerationTool.cpp
new file mode 100644
--- /dev/null
+++ b/Engine/Plugins/Enterprise/StaticMeshEditorExtension/Source/StaticMeshEditorExtension/Private/UVTools/SUVGenerationTool.cpp
@@ -0,0 +1,487 @@
+// Copyright 1998-2019 Epic Games, Inc. All Rights Reserved.
+
+#include "SUVGenerationTool.h"
+#include "UVGenerationTool.h"
+#include "UVGenerationToolbar.h"
+#include "IStaticMeshEditor.h"
+#include "Modules/ModuleManager.h"
+#include "Interfaces/IPluginManager.h"
+#include "StaticMeshAttributes.h"
+#include "MeshDescriptionOperations.h"
+#include "UVMapSettings.h"
+#include "ScopedTransaction.h"
+#include "DetailLayoutBuilder.h"
+#include "IStructureDetailsView.h"
+#include "Misc/MessageDialog.h"
+#include "Widgets/Layout/SBox.h"
+#include "Widgets/Input/SButton.h"
+#include "Widgets/Input/SMultiLineEditableTextBox.h"
+#include "Styling/SlateStyle.h"
+#include "Styling/SlateStyleRegistry.h"
+#include "EditorViewportClient.h"
+#include "EditorModeManager.h"
+
+#define LOCTEXT_NAMESPACE "UVGenerationTool"
+
+void UGenerateUVSettingsUIHolder::PostEditUndo()
+{
+	UObject::PostEditUndo();
+	OnUVSettingsRefreshNeededEvent.Broadcast();
+}
+
+SGenerateUV::SGenerateUV()
+{
+	SettingObjectUIHolder = NewObject<UGenerateUVSettingsUIHolder>((UObject*)GetTransientPackage(), NAME_None, EObjectFlags::RF_Transactional);
+	SettingObjectUIHolder->AddToRoot();
+	GenerateUVSettings = &SettingObjectUIHolder->GenerateUVSettings;
+}
+
+SGenerateUV::~SGenerateUV()
+{
+	SetPreviewModeActivated(false);
+	SettingObjectUIHolder->OnUVSettingsRefreshNeeded().RemoveAll(this);
+	SettingObjectUIHolder->RemoveFromRoot();
+
+	if (StaticMeshEditorPtr.IsValid())
+	{
+		StaticMeshEditorPtr.Pin()->UnRegisterOnSelectedLODChanged(this);
+	}
+}
+
+void SGenerateUV::Construct(const FArguments& InArgs)
+{
+	UVGenerationTool = InArgs._UVGenerationTool;
+	StaticMeshEditorPtr = InArgs._StaticMeshEditorPtr;
+	TSharedPtr<SBox> InspectorBox;
+
+	FitSettings();
+	SetNextValidTargetChannel();
+	if (!bAreDelegatesRegistered)
+	{
+		bAreDelegatesRegistered = true;
+		SettingObjectUIHolder->OnUVSettingsRefreshNeeded().AddSP(this, &SGenerateUV::UpdateUVPreview);
+		GenerateUVSettings->OnShapeEditingValueChanged.AddSP(this, &SGenerateUV::UpdateUVPreview);
+		GenerateUVSettings->OnGetNumberOfUVs.BindSP(this, &SGenerateUV::GetNumberOfUVChannels);
+	}
+
+	this->ChildSlot
+	[
+		SNew(SVerticalBox)
+
+		+ SVerticalBox::Slot()
+		.AutoHeight()
+		[
+			SAssignNew(ErrorHintWidget, SMultiLineEditableTextBox)
+			.Visibility(TAttribute<EVisibility>::Create(TAttribute<EVisibility>::FGetter::CreateLambda([this]() { return !IsCustomUVInspectorBoxEnabled() || IsTargetingLightMapUVChannel() ? EVisibility::Visible : EVisibility::Collapsed; })))
+			.Font(IDetailLayoutBuilder::GetDetailFontBold())
+			.BackgroundColor(TAttribute<FSlateColor>::Create(TAttribute<FSlateColor>::FGetter::CreateLambda([]() { return FEditorStyle::GetColor("ErrorReporting.WarningBackgroundColor"); })))
+			.Text(TAttribute<FText>::Create(TAttribute<FText>::FGetter::CreateSP(this, &SGenerateUV::GetWarningText)))
+			.AutoWrapText(true)
+			.IsReadOnly(true)
+		]
+
+		+ SVerticalBox::Slot()
+		.AutoHeight()
+		[
+			SNew(SVerticalBox)
+			.IsEnabled(this, &SGenerateUV::IsCustomUVInspectorBoxEnabled)
+
+			+ SVerticalBox::Slot()
+			.AutoHeight()
+			[
+				SAssignNew(InspectorBox, SBox)
+			]
+
+			+ SVerticalBox::Slot()
+			.AutoHeight()
+			.HAlign(HAlign_Center)
+			.Padding(0.0f, 8.0f, 0.0f, 0.0f)
+			[
+				SNew(SHorizontalBox)
+
+				+ SHorizontalBox::Slot()
+				.AutoWidth()
+				.Padding(8.0f, 0.0f, 8.0f, 0.0f)
+				[
+					SNew(SButton)
+					.VAlign(VAlign_Center)
+					.OnClicked(this, &SGenerateUV::OnShowGizmoButtonPressed)
+					.ButtonColorAndOpacity(this, &SGenerateUV::ShowGizmoButtonColorAndOpacity)
+					.Text(LOCTEXT("ShowGizmo_GenerateUV", "Show gizmo"))
+				]
+
+				+ SHorizontalBox::Slot()
+				.AutoWidth()
+				.Padding(8.0f, 0.0f, 8.0f, 0.0f)
+				[
+					SNew(SButton)
+					.VAlign(VAlign_Center)
+					.Text(LOCTEXT("Apply_GenerateUV", "Apply"))
+					.ToolTipText(LOCTEXT("Apply_GenerateUV_Tooltip", "Apply the generated UV to the target channel."))
+					.OnClicked(this, &SGenerateUV::OnApplyUV)
+					.IsEnabled(TAttribute<bool>::Create(TAttribute<bool>::FGetter::CreateLambda([this]() {return !IsTargetingLightMapUVChannel(); })))
+				]
+
+				+ SHorizontalBox::Slot()
+				.AutoWidth()
+				.Padding(8.0f, 0.0f, 8.0f, 0.0f)
+				[
+					SNew(SButton)
+					.VAlign(VAlign_Center)
+					.Text(LOCTEXT("Fit_GenerateUV", "Fit"))
+					.ToolTipText(LOCTEXT("Fix_GenerateUV_Tooltip", "Automatically sets the projection settings so that the generated UV fits properly."))
+					.OnClicked_Lambda([this]() -> FReply
+					{
+						FitSettings();
+						return FReply::Handled();
+					})
+				]
+			]
+		]
+	];
+
+	FPropertyEditorModule& PropertyEditorModule = FModuleManager::GetModuleChecked<FPropertyEditorModule>("PropertyEditor");
+	FDetailsViewArgs DetailsViewArgs;
+	DetailsViewArgs.bAllowSearch = false;
+	DetailsViewArgs.NameAreaSettings = FDetailsViewArgs::HideNameArea;
+	DetailsView = PropertyEditorModule.CreateDetailView(DetailsViewArgs);
+	DetailsView->SetObject(SettingObjectUIHolder);
+	DetailsView->OnFinishedChangingProperties().AddSP(this, &SGenerateUV::OnFinishedChangingProjectionPropertiesDetailView);
+	InspectorBox->SetContent(DetailsView->AsShared());
+}
+
+int32 GetSelectedLOD(const TSharedPtr<IStaticMeshEditor>& EditorPtr)
+{
+	const int32 SelectedLOD = EditorPtr->GetStaticMeshComponent()->ForcedLodModel - 1;
+
+	//If there is only one LOD then there is no ambiguity, even in AutoLOD.
+	return EditorPtr->GetStaticMesh()->GetNumLODs() == 1 ? 0 : SelectedLOD;
+}
+
+bool IsLODGenerated(TSharedPtr<IStaticMeshEditor> EditorPtr, int32 LODIndex)
+{
+	const UStaticMesh* StaticMesh = EditorPtr->GetStaticMesh();
+	return !StaticMesh->IsMeshDescriptionValid(LODIndex) || StaticMesh->IsReductionActive(LODIndex);
+}
+
+bool SGenerateUV::IsTargetingLightMapUVChannel() const
+{
+	if (StaticMeshEditorPtr.IsValid())
+	{
+		TSharedPtr<IStaticMeshEditor> EditorPtr = StaticMeshEditorPtr.Pin();
+		const int32 SelectedLOD = GetSelectedLOD(EditorPtr);
+		const FMeshBuildSettings& LODBuildSettings = EditorPtr->GetStaticMesh()->GetSourceModel(SelectedLOD).BuildSettings;
+		return LODBuildSettings.DstLightmapIndex == GenerateUVSettings->TargetChannel;
+	}
+
+	return false;
+}
+
+bool SGenerateUV::IsCustomUVInspectorBoxEnabled() const
+{
+	if (StaticMeshEditorPtr.IsValid())
+	{
+		TSharedPtr<IStaticMeshEditor> EditorPtr = StaticMeshEditorPtr.Pin();
+		const int32 SelectedLOD = GetSelectedLOD(EditorPtr);
+		bool bEnabled = SelectedLOD >= 0 && !IsLODGenerated(EditorPtr, SelectedLOD);
+
+		return bEnabled;
+	}
+
+	return false;
+}
+
+FVector SGenerateUV::ConvertTranslationToUIFormat(const FVector& InTranslation, const FRotator& InRotation) const
+{
+	//STR transform to SRT transform
+	FTransform ShapeFittingTransform = FTransform(InTranslation) * FTransform(InRotation.GetInverse());
+	return ShapeFittingTransform.GetTranslation();
+}
+
+FVector SGenerateUV::ConvertTranslationToAPIFormat(const FVector& InTranslation, const FRotator& InRotation) const
+{
+	//SRT transform to STR transform
+	FTransform ShapeFittingTransform = FTransform(InTranslation) * FTransform(InRotation);
+	return ShapeFittingTransform.GetTranslation();
+}
+
+bool SGenerateUV::GenerateUVTexCoords(TMap<FVertexInstanceID, FVector2D>& OutTexCoords) const
+{
+	OutTexCoords.Empty();
+
+	if (StaticMeshEditorPtr.IsValid())
+	{
+		TSharedPtr<IStaticMeshEditor> EditorPtr = StaticMeshEditorPtr.Pin();
+		UStaticMesh* StaticMesh = EditorPtr->GetStaticMesh();
+		const int32 CurrentLOD = GetSelectedLOD(EditorPtr);
+		FMeshDescription* MeshDescription = StaticMesh->GetMeshDescription(CurrentLOD);
+		FRotator InvertedRotation = GenerateUVSettings->Rotation.GetInverse();
+		FVector SRTPosition = ConvertTranslationToAPIFormat(GenerateUVSettings->Position, InvertedRotation);
+
+		//The rotation of the FUVMapParameters is actually a rotation applied to the mesh and not the projection shape, that means we must get the Inverse rotation of the shape.
+		FUVMapParameters UVParameters(SRTPosition, InvertedRotation.Quaternion(), GenerateUVSettings->Size, FVector::OneVector, GenerateUVSettings->UVTilingScale);
+
+		switch (GenerateUVSettings->ProjectionType)
+		{
+		case EGenerateUVProjectionType::Box:
+			FMeshDescriptionOperations::GenerateBoxUV(*MeshDescription, UVParameters, OutTexCoords);
+			break;
+		case EGenerateUVProjectionType::Cylindrical:
+			FMeshDescriptionOperations::GenerateCylindricalUV(*MeshDescription, UVParameters, OutTexCoords);
+			break;
+		case EGenerateUVProjectionType::Planar:
+			FMeshDescriptionOperations::GeneratePlanarUV(*MeshDescription, UVParameters, OutTexCoords);
+			break;
+		}
+
+		for (TPair<FVertexInstanceID, FVector2D>& Pair : OutTexCoords)
+		{
+			Pair.Value.X += GenerateUVSettings->UVOffset.X;
+			Pair.Value.Y += GenerateUVSettings->UVOffset.Y;
+		}
+
+		return true;
+	}
+
+	return false;
+}
+
+FReply SGenerateUV::OnApplyUV() const
+{
+	TMap<FVertexInstanceID, FVector2D> TexCoords;
+	TSharedPtr<IStaticMeshEditor> EditorPtr = StaticMeshEditorPtr.Pin();
+
+	if (EditorPtr.IsValid() && GenerateUVTexCoords(TexCoords))
+	{
+		UStaticMesh* StaticMesh = EditorPtr->GetStaticMesh();
+		const int32 CurrentLOD = GetSelectedLOD(EditorPtr);
+		const FMeshDescription* MeshDescription = StaticMesh->GetMeshDescription(CurrentLOD);
+		FStaticMeshConstAttributes Attributes(*MeshDescription);
+		const int32 NumberOfUVChannels = Attributes.GetVertexInstanceUVs().GetNumIndices();
+		int32 TargetChannel = FMath::Min((int32)GenerateUVSettings->TargetChannel, EditorPtr->GetNumUVChannels());
+		const int32 NumberOfChannelsToAdd = TargetChannel - (NumberOfUVChannels - 1);
+
+		FScopedTransaction Transaction(LOCTEXT("ApplyGeneratedUVTransaction", "Applied Generated UV"));
+		if (NumberOfChannelsToAdd > 0)
+		{
+			for (int Index = 0; Index < NumberOfChannelsToAdd; ++Index)
+			{
+				StaticMesh->AddUVChannel(CurrentLOD);
+			}
+		}
+		else
+		{
+			//Ask for user confirmation before overwriting existing UV channel.
+			const EAppReturnType::Type UserResponse = FMessageDialog::Open(
+				EAppMsgType::YesNo,
+				LOCTEXT("ApplyOnExistingUVConfirmation", "An UV channel already exists at this index.\n\nDo you want to overwrite the existing channel data?"));
+
+			if (UserResponse == EAppReturnType::No)
+			{
+				return FReply::Handled();
+			}
+		}
+		StaticMesh->SetUVChannel(CurrentLOD, TargetChannel, TexCoords);
+		EditorPtr->RefreshTool();
+		DetailsView->ForceRefresh();
+	}
+
+	return FReply::Handled();
+}
+
+
+FReply SGenerateUV::OnShowGizmoButtonPressed()
+{
+	SetPreviewModeActivated(!bIsInPreviewUVMode);
+
+	return FReply::Handled();
+}
+
+void SGenerateUV::SetPreviewModeActivated(bool bActive)
+{
+	bIsInPreviewUVMode = bActive;
+	if (!StaticMeshEditorPtr.IsValid())
+	{
+		return;
+	}
+
+	FEditorViewportClient& ViewportClient = StaticMeshEditorPtr.Pin()->GetViewportClient();
+	FEditorModeTools* ModeTools = ViewportClient.GetModeTools();
+
+	if (bIsInPreviewUVMode)
+	{
+		ModeTools->ActivateMode(FUVGenerationTool::EM_UVGeneration);
+		GenerateUVPreviewMode = StaticCastSharedRef<FUVGenerationTool>(ModeTools->GetActiveMode(FUVGenerationTool::EM_UVGeneration)->AsShared());
+		OnWidgetChangedShapeSettingsHandle = GenerateUVPreviewMode.Pin()->OnShapeSettingsChanged().AddSP(this, &SGenerateUV::OnWidgetChangedShapeSettings);
+		OnEditorModeChangedHandle = ModeTools->OnEditorModeIDChanged().AddSP(this, &SGenerateUV::OnEditorModeChanged);
+		UpdateUVPreview();
+	}
+	else
+	{
+		OnGenerateUVPreviewModeDeactivated();
+		ModeTools->DeactivateMode(FUVGenerationTool::EM_UVGeneration);	
+	}
+}
+
+void SGenerateUV::OnEditorModeChanged(const FEditorModeID& ModeChangedID, bool bIsEnteringMode)
+{
+	//If the UV shape preview mode is getting deactivated by something.
+	if (bIsInPreviewUVMode && !bIsEnteringMode && GenerateUVPreviewMode.IsValid() && GenerateUVPreviewMode.Pin().Get()->GetID() == ModeChangedID)
+	{
+		OnGenerateUVPreviewModeDeactivated();
+	}
+}
+
+void SGenerateUV::OnGenerateUVPreviewModeDeactivated()
+{
+	if (StaticMeshEditorPtr.IsValid())
+	{
+		FEditorViewportClient& ViewportClient = StaticMeshEditorPtr.Pin()->GetViewportClient();
+		FEditorModeTools* ModeTools = ViewportClient.GetModeTools();
+
+		if (OnWidgetChangedShapeSettingsHandle.IsValid() && GenerateUVPreviewMode.IsValid())
+		{
+			GenerateUVPreviewMode.Pin()->OnShapeSettingsChanged().Remove(OnWidgetChangedShapeSettingsHandle);
+			OnWidgetChangedShapeSettingsHandle.Reset();
+		}
+
+		if (OnEditorModeChangedHandle.IsValid() && ModeTools)
+		{
+			ModeTools->OnEditorModeIDChanged().Remove(OnEditorModeChangedHandle);
+			OnEditorModeChangedHandle.Reset();
+		}
+	}
+
+	bIsInPreviewUVMode = false;
+	GenerateUVPreviewMode.Reset();
+}
+
+FSlateColor SGenerateUV::ShowGizmoButtonColorAndOpacity() const
+{
+	static const FName SelectionColor("SelectionColor");
+
+	return bIsInPreviewUVMode ? FCoreStyle::Get().GetSlateColor(SelectionColor) : FLinearColor::White;
+}
+
+void SGenerateUV::OnFinishedChangingProjectionPropertiesDetailView(const FPropertyChangedEvent& PropertyChangedEvent) const
+{
+	UpdateUVPreview();
+}
+
+void SGenerateUV::OnWidgetChangedShapeSettings(const FVector& Position, const FVector& Size, const FRotator& Rotation)
+{
+	SettingObjectUIHolder->Modify();
+
+	GenerateUVSettings->Position = Position;
+	GenerateUVSettings->Size = Size;
+	GenerateUVSettings->Rotation = Rotation;
+}
+
+void SGenerateUV::UpdateUVPreview() const
+{
+	if (GenerateUVPreviewMode.IsValid())
+	{
+		GenerateUVPreviewMode.Pin()->SetShapeSettings(*GenerateUVSettings);
+	}
+}
+
+int32 SGenerateUV::GetNumberOfUVChannels()
+{
+	if (StaticMeshEditorPtr.IsValid())
+	{
+		return StaticMeshEditorPtr.Pin()->GetNumUVChannels();
+	}
+
+	return 0;
+}
+
+void SGenerateUV::SetNextValidTargetChannel()
+{
+	if (StaticMeshEditorPtr.IsValid())
+	{
+		TSharedPtr<IStaticMeshEditor> EditorPtr = StaticMeshEditorPtr.Pin();
+		const int32 CurrentLOD = GetSelectedLOD(EditorPtr);
+
+		GenerateUVSettings->TargetChannel = FMath::Min(EditorPtr->GetNumUVChannels(CurrentLOD), MAX_MESH_TEXTURE_COORDS_MD - 1);
+	}
+}
+
+void SGenerateUV::FitSettings()
+{
+	TSharedPtr<IStaticMeshEditor> EditorPtr = StaticMeshEditorPtr.Pin();
+
+	if (EditorPtr.IsValid())
+	{
+		UStaticMesh* StaticMesh = StaticMeshEditorPtr.Pin()->GetStaticMesh();
+		
+		const int32 CurrentLOD = FMath::Max(0, GetSelectedLOD(EditorPtr));
+		const FMeshDescription* MeshDescription = StaticMesh->GetMeshDescription(CurrentLOD);
+		FStaticMeshConstAttributes Attributes(*MeshDescription);
+		TMeshAttributesConstRef<FVertexID, FVector> VertexPositions = Attributes.GetVertexPositions();
+		TArray<FVector> RotatedVertexPositions;
+		RotatedVertexPositions.Reserve(MeshDescription->VertexInstances().Num());
+		FRotator InvertedRotation = GenerateUVSettings->Rotation.GetInverse();
+
+		for (const FVertexInstanceID VertexInstanceID : MeshDescription->VertexInstances().GetElementIDs())
+		{
+			const FVertexID VertexID = MeshDescription->GetVertexInstanceVertex(VertexInstanceID);
+			RotatedVertexPositions.Emplace(InvertedRotation.RotateVector(VertexPositions[VertexID]));
+		}
+
+		//We set the size an position by using the rotated mesh bounds
+		FBox BoundsBox(RotatedVertexPositions);
+		GenerateUVSettings->Size = BoundsBox.GetExtent() * 2.f;
+		GenerateUVSettings->Position = ConvertTranslationToUIFormat(BoundsBox.GetCenter(), InvertedRotation);
+
+		if (GenerateUVSettings->ProjectionType == EGenerateUVProjectionType::Cylindrical)
+		{
+			//We need a scaling correction to take account of the BoundsBox corners with the Cylinder projection.
+			float MaxSqrSize2D = 0;
+			const FVector CylinderExtent = BoundsBox.GetExtent(), VertexOffset = BoundsBox.GetCenter();
+
+			for (const FVector CurrentVertex : RotatedVertexPositions)
+			{
+				FVector VertexRatio = (CurrentVertex - VertexOffset) / CylinderExtent;
+				float SqrSize2D = VertexRatio.Y * VertexRatio.Y + VertexRatio.Z * VertexRatio.Z;
+
+				MaxSqrSize2D = FMath::Max(SqrSize2D, MaxSqrSize2D);
+			}
+
+			if (MaxSqrSize2D > 1.f)
+			{
+				float MaxSize2D = FMath::Sqrt(MaxSqrSize2D);
+				FVector CorrectedExtent(CylinderExtent.X, CylinderExtent.Y * MaxSize2D, CylinderExtent.Z * MaxSize2D);
+				GenerateUVSettings->Size = CorrectedExtent * 2.f;
+			}
+		}
+
+		UpdateUVPreview();
+	}
+}
+
+FText SGenerateUV::GetWarningText()
+{
+	if (StaticMeshEditorPtr.IsValid())
+	{
+		//Enable the edition only if we have a specific LOD selected (0 == Auto_LOD, 1 == LOD_0, etc.)
+		TSharedPtr<IStaticMeshEditor> EditorPtr = StaticMeshEditorPtr.Pin();
+		const int32 SelectedLOD = GetSelectedLOD(EditorPtr);
+
+		if (SelectedLOD < 0)
+		{
+			return LOCTEXT("StaticMeshGenerateUV_Warning_AutoLOD", "Cannot apply custom UV in AutoLOD view mode.");
+		}
+		else if (IsLODGenerated(EditorPtr, SelectedLOD))
+		{
+			return LOCTEXT("StaticMeshGenerateUV_Warning_GeneratedLOD", "Cannot apply custom UV on a generated LOD.");
+		}
+		else if (IsTargetingLightMapUVChannel())
+		{
+			return LOCTEXT("StaticMeshGenerateUV_Warning_LightMapChannel", "Cannot apply custom UV in the lightmap UV channel.");
+		}
+	}
+
+	return FText();
+}
+
+#undef LOCTEXT_NAMESPACE
\ No newline at end of file
diff --git a/Engine/Plugins/Enterprise/StaticMeshEditorExtension/Source/StaticMeshEditorExtension/Private/UVTools/UVGenerationTool.cpp b/Engine/Plugins/Enterprise/StaticMeshEditorExtension/Source/StaticMeshEditorExtension/Private/UVTools/UVGenerationTool.cpp
new file mode 100644
--- /dev/null
+++ b/Engine/Plugins/Enterprise/StaticMeshEditorExtension/Source/StaticMeshEditorExtension/Private/UVTools/UVGenerationTool.cpp
@@ -0,0 +1,176 @@
+// Copyright 1998-2019 Epic Games, Inc. All Rights Reserved.
+
+#include "UVGenerationTool.h"
+#include "EditorModeManager.h"
+#include "EditorViewportClient.h"
+
+#define LOCTEXT_NAMESPACE "UVGenerationTool"
+
+const FEditorModeID FUVGenerationTool::EM_UVGeneration(TEXT("EM_UVGeneration"));
+
+bool FUVGenerationTool::GetCustomDrawingCoordinateSystem(FMatrix& InMatrix, void* InData)
+{
+	if (IsHandlingInputs())
+	{
+		InMatrix = FRotationMatrix::Make(ShapeRotation);
+		return true;
+	}
+
+	return false;
+}
+
+bool FUVGenerationTool::InputDelta(FEditorViewportClient* InViewportClient, FViewport* InViewport, FVector& InDrag, FRotator& InRot, FVector& InScale)
+{
+	if (!(IsHandlingInputs() && InViewportClient->bWidgetAxisControlledByDrag))
+	{
+		return false;
+	}
+
+	ShapePosition += InDrag;
+	ShapeSize += InScale;
+
+	if (InRot != FRotator::ZeroRotator)
+	{
+		FRotator RotationWinding, RotationRemainder;
+		ShapeRotation.GetWindingAndRemainder(RotationWinding, RotationRemainder);
+
+		const FQuat ActorQuat = RotationRemainder.Quaternion();
+		const FQuat DeltaQuat = InRot.Quaternion();
+		const FQuat ResultQuat = DeltaQuat * ActorQuat;
+		const FRotator NewSocketRotationRemainder = FRotator(ResultQuat);
+		FRotator DeltaRot = NewSocketRotationRemainder - RotationRemainder;
+		DeltaRot.Normalize();
+
+		ShapeRotation += DeltaRot;
+	}
+
+	OnShapeSettingsChangedEvent.Broadcast(ShapePosition, ShapeSize, ShapeRotation);
+	return true;
+}
+
+bool FUVGenerationTool::StartTracking(FEditorViewportClient* InViewportClient, FViewport* InViewport)
+{
+	if (IsHandlingInputs() && InViewportClient->bWidgetAxisControlledByDrag)
+	{
+		if (!bIsTrackingWidgetDrag)
+		{
+			const FWidget::EWidgetMode WidgetMode = GetModeManager()->GetWidgetMode();
+			FText TransText;
+
+			if (WidgetMode == FWidget::WM_Rotate)
+			{
+				TransText = LOCTEXT("ChangeRotationSettings", "Rotate projection shape");
+			}
+			else if(WidgetMode == FWidget::WM_Scale)
+			{
+				TransText = LOCTEXT("ChangeSizeSettings", "Scale projection shape");
+			}
+			else
+			{
+				TransText = LOCTEXT("ChangePositionSettings", "Move projection shape");
+			}
+
+			bIsTrackingWidgetDrag = true;
+			GEditor->BeginTransaction(TransText);
+		}
+
+		return true;
+	}
+
+	return false;
+}
+
+bool FUVGenerationTool::EndTracking(FEditorViewportClient* InViewportClient, FViewport* InViewport)
+{
+	if (bIsTrackingWidgetDrag)
+	{
+		bIsTrackingWidgetDrag = false;
+		GEditor->EndTransaction();
+
+		return true;
+	}
+
+	return false;
+}
+
+void FUVGenerationTool::Enter()
+{
+	PreviousWidgetMode = GetModeManager()->GetWidgetMode();
+	GetModeManager()->SetWidgetMode(FWidget::EWidgetMode::WM_Translate);
+}
+
+void FUVGenerationTool::Exit()
+{
+	GetModeManager()->SetWidgetMode(PreviousWidgetMode);
+}
+
+void FUVGenerationTool::Draw(const FSceneView* View, FPrimitiveDrawInterface* PDI)
+{
+	FLinearColor DrawColor = FLinearColor::Green;
+	FTransform Transform(ShapeRotation, FVector::ZeroVector, ShapeSize * 0.5f);
+
+	switch (ShapeType)
+	{
+		case EGenerateUVProjectionType::Box:
+		{
+			DrawOrientedWireBox(PDI, ShapePosition,
+				Transform.GetScaledAxis(EAxis::X), Transform.GetScaledAxis(EAxis::Y), Transform.GetScaledAxis(EAxis::Z),
+				FVector::OneVector, DrawColor, SDPG_Foreground);
+
+			break;
+		}
+		case EGenerateUVProjectionType::Cylindrical:
+		{
+			//The cylinder of the Cylinder Projection is aligned on the X-axis, but the DrawWireCylinder() is aligned on the Z-axis so we apply a rotation and shape correction.
+			FVector BoxSize = ShapeSize * 0.5f;
+			Transform = FTransform((FRotator(90, 0, 0).Quaternion() * ShapeRotation.Quaternion()), FVector::ZeroVector, FVector(BoxSize.Z, BoxSize.Y, BoxSize.X));
+
+			DrawWireCylinder(PDI, ShapePosition,
+				Transform.GetScaledAxis(EAxis::X), Transform.GetScaledAxis(EAxis::Y), Transform.GetScaledAxis(EAxis::Z),
+				DrawColor, 1, 1, 20, SDPG_Foreground);
+
+			break;
+		}
+		case EGenerateUVProjectionType::Planar :
+		{
+			//Simple Plane drawing.
+			FVector QuadVertices[4] = {
+				FVector(0, -1, 1),
+				FVector(0, 1, 1),
+				FVector(0, -1, -1),
+				FVector(0, 1, -1)
+			};
+
+			for (int32 VertexIndex = 0; VertexIndex < 4; ++VertexIndex)
+			{
+				QuadVertices[VertexIndex] = ShapePosition + Transform.TransformPosition(QuadVertices[VertexIndex]);
+			}
+
+			PDI->DrawLine(QuadVertices[0], QuadVertices[1], DrawColor, SDPG_Foreground, 0, 0, false);
+			PDI->DrawLine(QuadVertices[0], QuadVertices[2], DrawColor, SDPG_Foreground, 0, 0, false);
+			PDI->DrawLine(QuadVertices[1], QuadVertices[3], DrawColor, SDPG_Foreground, 0, 0, false);
+			PDI->DrawLine(QuadVertices[2], QuadVertices[3], DrawColor, SDPG_Foreground, 0, 0, false);
+
+			break;
+		}
+	}
+}
+
+void FUVGenerationTool::SetShapeSettings(const FUVGenerationSettings& GenerationSettings)
+{
+	ShapeType = GenerationSettings.ProjectionType;
+	ShapePosition = GenerationSettings.Position;
+	ShapeSize = GenerationSettings.Size;
+	ShapeRotation = GenerationSettings.Rotation;
+}
+
+bool FUVGenerationTool::IsHandlingInputs() const 
+{
+	const ECoordSystem DeltaCoordSystem = GetModeManager()->GetCoordSystem();
+	const FWidget::EWidgetMode WidgetMode = GetModeManager()->GetWidgetMode();
+
+	return (DeltaCoordSystem == ECoordSystem::COORD_Local || DeltaCoordSystem == COORD_World) &&
+		(WidgetMode == FWidget::WM_Translate ||	WidgetMode == FWidget::WM_Rotate ||	WidgetMode == FWidget::WM_Scale);
+}
+
+#undef LOCTEXT_NAMESPACE
\ No newline at end of file
diff --git a/Engine/Plugins/Experimental/DataPrepEditor/Source/DataPrepCore/Private/Blueprint/K2Node_DataprepProducer.cpp b/Engine/Plugins/Experimental/DataPrepEditor/Source/DataPrepCore/Private/Blueprint/K2Node_DataprepProducer.cpp
new file mode 100644
--- /dev/null
+++ b/Engine/Plugins/Experimental/DataPrepEditor/Source/DataPrepCore/Private/Blueprint/K2Node_DataprepProducer.cpp
@@ -0,0 +1,58 @@
+// Copyright 1998-2019 Epic Games, Inc. All Rights Reserved.
+
+#include "K2Node_DataprepProducer.h"
+#include "DataPrepContentProducer.h"
+
+#include "EdGraphSchema_K2.h"
+#include "SGraphNode.h"
+
+#include "Widgets/DeclarativeSyntaxSupport.h"
+
+void UK2Node_DataprepProducer::AllocateDefaultPins()
+{
+	// The immediate continue pin
+	CreatePin(EGPD_Output, UEdGraphSchema_K2::PC_Exec, UEdGraphSchema_K2::PN_Then);
+	Super::AllocateDefaultPins();
+}
+
+FText UK2Node_DataprepProducer::GetNodeTitle(ENodeTitleType::Type TitleType) const
+{
+	return NSLOCTEXT("UK2Node_DataprepProducer", "NodeTitle", "Start");
+}
+
+FText UK2Node_DataprepProducer::GetTooltipText() const
+{
+	return NSLOCTEXT("UK2Node_DataprepProducer", "NodeTooltip", "Hold onto all the action assets associated to a Blueprint based Dataprep asset");
+}
+
+#if WITH_EDITOR
+TSharedPtr<SGraphNode> UK2Node_DataprepProducer::CreateVisualWidget()
+{
+	typedef TTuple< UClass*, FText, FText > DataprepProducerDescription;
+
+	class SGraphNodeDataprepProducer : public SGraphNode
+	{
+	public:
+		SLATE_BEGIN_ARGS(SGraphNodeDataprepProducer){}
+		SLATE_END_ARGS()
+
+			SGraphNodeDataprepProducer()
+		{
+		}
+
+		~SGraphNodeDataprepProducer()
+		{
+		}
+
+		void Construct(const FArguments& InArgs, UK2Node_DataprepProducer* InNode)
+		{
+			GraphNode = InNode;
+			UpdateGraphNode();
+		}
+	};
+
+	return SNew(SGraphNodeDataprepProducer, this);
+}
+#endif
+
+#undef LOCTEXT_NAMESPACE
diff --git a/Engine/Plugins/Experimental/DataPrepEditor/Source/DataPrepEditor/Private/BlueprintNodes/K2Node_DataprepProducer.cpp b/Engine/Plugins/Experimental/DataPrepEditor/Source/DataPrepEditor/Private/BlueprintNodes/K2Node_DataprepProducer.cpp
deleted file mode 100644
diff --git a/Engine/Source/Developer/Datasmith/DatasmithFacade/Private/DatasmithFacadeActorCamera.cpp b/Engine/Source/Developer/Datasmith/DatasmithFacade/Private/DatasmithFacadeActorCamera.cpp
new file mode 100644
--- /dev/null
+++ b/Engine/Source/Developer/Datasmith/DatasmithFacade/Private/DatasmithFacadeActorCamera.cpp
@@ -0,0 +1,151 @@
+// Copyright 1998-2019 Epic Games, Inc. All Rights Reserved.
+
+#include "DatasmithFacadeActorCamera.h"
+
+
+FDatasmithFacadeActorCamera::FDatasmithFacadeActorCamera(
+	const TCHAR* InElementName,
+	const TCHAR* InElementLabel
+) :
+	FDatasmithFacadeActor(InElementName, InElementLabel),
+	WorldCameraPosition(0.0, 0.0, 0.0), // default camera position at world origin
+	SensorWidth(36.0),                  // default Datasmith sensor width of 36 mm
+	AspectRatio(16.0 / 9.0),            // default Datasmith aspect ratio of 16:9
+	FocusDistance(1000.0),              // default Datasmith focus distance of 1000 centimeters
+	FocalLength(35.0),                  // default Datasmith focal length of 35 millimeters
+	LookAtActorName(),					// default Datasmith look-at actor name
+	bLookAtAllowRoll(false)				// default Datasmith look-at allow roll state
+{
+	// Prevent the Datasmith camera actor from being removed by optimization.
+	KeepActor();
+}
+
+void FDatasmithFacadeActorCamera::SetCameraPosition(
+	float InX,
+	float InY,
+	float InZ
+)
+{
+	// Scale and convert the camera position into a Datasmith actor world translation.
+	WorldTransform.SetTranslation(ConvertPosition(InX, InY, InZ));
+}
+
+void FDatasmithFacadeActorCamera::SetCameraRotation(
+	float InForwardX,
+	float InForwardY,
+	float InForwardZ,
+	float InUpX,
+	float InUpY,
+	float InUpZ
+)
+{
+	// Convert the camera orientation into a Datasmith actor world look-at rotation quaternion.
+	FVector XAxis = ConvertDirection(InForwardX, InForwardY, InForwardZ);
+	FVector ZAxis = ConvertDirection(InUpX, InUpY, InUpZ);
+	WorldTransform.SetRotation(FQuat(FRotationMatrix::MakeFromXZ(XAxis, ZAxis))); // axis vectors do not need to be normalized
+}
+
+void FDatasmithFacadeActorCamera::SetSensorWidth(
+	float InSensorWidth
+)
+{
+	SensorWidth = InSensorWidth;
+}
+
+void FDatasmithFacadeActorCamera::SetAspectRatio(
+	float InAspectRatio
+)
+{
+	AspectRatio = InAspectRatio;
+}
+
+void FDatasmithFacadeActorCamera::SetFocusDistance(
+	float InTargetX,
+	float InTargetY,
+	float InTargetZ
+)
+{
+	// Set the Datasmith camera focus distance (in centimeters).
+	FVector DistanceVector(InTargetX - WorldCameraPosition.X, InTargetY - WorldCameraPosition.Y, InTargetZ - WorldCameraPosition.Z);
+	FocusDistance = DistanceVector.Size() * WorldUnitScale;
+}
+
+void FDatasmithFacadeActorCamera::SetFocusDistance(
+	float InFocusDistance
+)
+{
+	// Set the Datasmith camera focus distance (in centimeters).
+	FocusDistance = InFocusDistance * WorldUnitScale;
+}
+
+void FDatasmithFacadeActorCamera::SetFocalLength(
+	float InFocalLength
+)
+{
+	FocalLength = InFocalLength;
+}
+
+void FDatasmithFacadeActorCamera::SetFocalLength(
+	float InFOV,
+	bool  bInVerticalFOV
+)
+{
+	FocalLength = float((bInVerticalFOV ? (SensorWidth / AspectRatio) : SensorWidth) / (2.0 * tan(FMath::DegreesToRadians<double>(InFOV) / 2.0)));
+}
+
+void FDatasmithFacadeActorCamera::SetLookAtActor(
+	const TCHAR* InActorName
+)
+{
+	LookAtActorName = InActorName;
+}
+
+void FDatasmithFacadeActorCamera::SetLookAtAllowRoll(
+	bool bInAllow
+)
+{
+	bLookAtAllowRoll = bInAllow;
+}
+
+TSharedPtr<IDatasmithActorElement> FDatasmithFacadeActorCamera::CreateActorHierarchy(
+	TSharedRef<IDatasmithScene> IOSceneRef
+) const
+{
+	// Create a Datasmith camera actor element.
+	TSharedPtr<IDatasmithCameraActorElement> CameraActorPtr = FDatasmithSceneFactory::CreateCameraActor(*ElementName);
+
+	// Set the Datasmith camera actor base properties.
+	SetActorProperties(IOSceneRef, CameraActorPtr);
+
+	// Set the Datasmith camera sensor width.
+	CameraActorPtr->SetSensorWidth(SensorWidth);
+
+	// Set the Datasmith camera aspect ratio.
+	CameraActorPtr->SetSensorAspectRatio(AspectRatio);
+
+	// Set the Datasmith camera focus method.
+	CameraActorPtr->SetEnableDepthOfField(false); // no depth of field
+
+	// Set the Datasmith camera focus distance.
+	CameraActorPtr->SetFocusDistance(FocusDistance);
+
+	// Set the Datasmith camera f-stop.
+	CameraActorPtr->SetFStop(5.6); // default Datasmith f-stop of f/5.6
+
+	// Set the Datasmith camera focal length.
+	CameraActorPtr->SetFocalLength(FocalLength);
+
+	// Set the Datasmith camera look-at actor.
+	if (!LookAtActorName.IsEmpty())
+	{
+		CameraActorPtr->SetLookAtActor(*LookAtActorName);
+	}
+
+	// Set the Datasmith camera look-at allow roll state.
+	CameraActorPtr->SetLookAtAllowRoll(bLookAtAllowRoll);
+
+	// Add the hierarchy of children to the Datasmith actor.
+	AddActorChildren(IOSceneRef, CameraActorPtr);
+
+	return CameraActorPtr;
+}
