Fixed crash in Jacketing feature

#rnx

diff --git a/Engine/Plugins/Enterprise/StaticMeshEditorExtension/Source/MeshProcessingLibrary/Private/MeshProcessingLibrary.cpp b/Engine/Plugins/Enterprise/StaticMeshEditorExtension/Source/MeshProcessingLibrary/Private/MeshProcessingLibrary.cpp
--- a/Engine/Plugins/Enterprise/StaticMeshEditorExtension/Source/MeshProcessingLibrary/Private/MeshProcessingLibrary.cpp
+++ b/Engine/Plugins/Enterprise/StaticMeshEditorExtension/Source/MeshProcessingLibrary/Private/MeshProcessingLibrary.cpp
@@ -654,10 +654,13 @@ void UMeshProcessingLibrary::ApplyJacketingOnMeshActors(const TArray<AActor*>& A
 			TVertexInstanceAttributesRef<FVector> VertexInstanceTangents = NewRawMesh.VertexInstanceAttributes().GetAttributesRef<FVector>(MeshAttribute::VertexInstance::Tangent);
 
 			RawMesh->Empty();
-			// Update mesh description of static mesh with new geometry
-			FMeshDescription* MeshDescription = StaticMesh->GetMeshDescription(0);
-			//Copy the new mesh in the Original meshDescription hold from the staticmesh
-			*MeshDescription = NewRawMesh;
+
+			//Create the missing normals and tangents on polygons because FStaticMeshOperations::ComputeTangentsAndNormals requires it
+			if(!NewRawMesh.PolygonAttributes().GetAttributesRef<FVector>(MeshAttribute::Polygon::Normal).IsValid()
+				|| !NewRawMesh.PolygonAttributes().GetAttributesRef<FVector>(MeshAttribute::Polygon::Tangent).IsValid())
+			{
+				FStaticMeshOperations::ComputePolygonTangentsAndNormals(NewRawMesh);
+			}
 
 			const FMeshBuildSettings& BuildSettings = StaticMesh->GetSourceModel(0).BuildSettings;
 			EComputeNTBsFlags ComputeNTBsOptions = EComputeNTBsFlags::BlendOverlappingNormals;
@@ -667,11 +670,12 @@ void UMeshProcessingLibrary::ApplyJacketingOnMeshActors(const TArray<AActor*>& A
 			ComputeNTBsOptions |= BuildSettings.bComputeWeightedNormals ? EComputeNTBsFlags::WeightedNTBs : EComputeNTBsFlags::None;
 			ComputeNTBsOptions |= BuildSettings.bRemoveDegenerates ? EComputeNTBsFlags::IgnoreDegenerateTriangles : EComputeNTBsFlags::None;
 
-			FStaticMeshOperations::ComputeTangentsAndNormals(*MeshDescription, ComputeNTBsOptions);
+			FStaticMeshOperations::ComputeTangentsAndNormals(NewRawMesh, ComputeNTBsOptions);
 			// TODO: Maybe add generation of lightmap UV here.
 
-			//Commit the result so the old FRawMesh is updated
-			//StaticMesh->CommitMeshDescription(0);
+			// Update mesh description of static mesh with new geometry
+			FMeshDescription* MeshDescription = StaticMesh->GetMeshDescription(0);
+			*MeshDescription = MoveTemp(NewRawMesh);
 		}
 	});
 
