Path fixes for ThirdParty dependencies.

diff --git a/Engine/Plugins/Enterprise/DatasmithIFCImporter/Source/DatasmithIFCTranslator/DatasmithIFCTranslator.Build.cs b/Engine/Plugins/Enterprise/DatasmithIFCImporter/Source/DatasmithIFCTranslator/DatasmithIFCTranslator.Build.cs
--- a/Engine/Plugins/Enterprise/DatasmithIFCImporter/Source/DatasmithIFCTranslator/DatasmithIFCTranslator.Build.cs
+++ b/Engine/Plugins/Enterprise/DatasmithIFCImporter/Source/DatasmithIFCTranslator/DatasmithIFCTranslator.Build.cs
@@ -33,7 +33,7 @@ namespace UnrealBuildTool.Rules
 				PrivateDependencyModuleNames.Add("MessageLog");
 			}
 
-			string IfcEngineDir = Path.Combine(EngineDirectory, "Restricted/NotForLicensees/Plugins/Enterprise/DatasmithIFCTranslator/Source/ThirdParty/ifcengine");
+			string IfcEngineDir = Path.Combine(EngineDirectory, "Restricted/NotForLicensees/Source/ThirdParty/Enterprise/ifcengine");
 			if (Directory.Exists(IfcEngineDir))
 			{
 				PrivateDependencyModuleNames.AddRange(
diff --git a/Engine/Plugins/Enterprise/MDLImporter/Source/MDLImporter/MDLImporter.Build.cs b/Engine/Plugins/Enterprise/MDLImporter/Source/MDLImporter/MDLImporter.Build.cs
--- a/Engine/Plugins/Enterprise/MDLImporter/Source/MDLImporter/MDLImporter.Build.cs
+++ b/Engine/Plugins/Enterprise/MDLImporter/Source/MDLImporter/MDLImporter.Build.cs
@@ -11,7 +11,7 @@ namespace UnrealBuildTool.Rules
     {
         private string ThirdPartyPath
         {
-            get { return Path.GetFullPath(Path.Combine(EngineDirectory, "Restricted/NotForLicensees/Plugins/Enterprise/MDLImporter/Source/ThirdParty/")); }
+            get { return Path.GetFullPath(Path.Combine(EngineDirectory, "Restricted/NotForLicensees/Source/ThirdParty/Enterprise/")); }
         }
 
         public MDLImporter(ReadOnlyTargetRules Target) : base(Target)
diff --git a/Engine/Plugins/Enterprise/StaticMeshEditorExtension/Source/MeshProcessingLibrary/MeshProcessingLibrary.Build.cs b/Engine/Plugins/Enterprise/StaticMeshEditorExtension/Source/MeshProcessingLibrary/MeshProcessingLibrary.Build.cs
--- a/Engine/Plugins/Enterprise/StaticMeshEditorExtension/Source/MeshProcessingLibrary/MeshProcessingLibrary.Build.cs
+++ b/Engine/Plugins/Enterprise/StaticMeshEditorExtension/Source/MeshProcessingLibrary/MeshProcessingLibrary.Build.cs
@@ -42,7 +42,7 @@ public class MeshProcessingLibrary : ModuleRules
 		}
 
 		// Setup MeshSimplifier
-		string MeshSimplifierPath = Path.Combine(EngineDirectory, "Restricted/NotForLicensees/Plugins/Enterprise/MeshProcessingLibrary/Source/ThirdParty/MeshSimplifier");
+		string MeshSimplifierPath = Path.Combine(EngineDirectory, "Restricted/NotForLicensees/Source/ThirdParty/Enterprise/MeshSimplifier");
 		bool bWithMeshSimplifier = Directory.Exists(MeshSimplifierPath) && Target.Platform == UnrealTargetPlatform.Win64;
 		PublicDefinitions.Add("WITH_MESH_SIMPLIFIER=" + (bWithMeshSimplifier ? '1' : '0'));
 
