First pass of MeshDescription API and format refactor.
- Removed hardcoded element type arrays (Vertices, Edges, Triangles etc.). Mesh element types can now be arbitrarily added, with any number of channels.
- Mesh element containers have a much leaner format; instead of sparse arrays, they are now represented by a simple bitarray, determining whether an index is used or not. Consequently, mesh topology is now entirely described with the attribute system, e.g. edge start and end vertices, triangle vertices, etc.
- Support added for attributes of arbitrary dimensions, e.g. float[4] or int[2].
- Support added for attributes which index into another mesh element container.
- Added FMeshElementIndexer: this is an efficient container for maintaining backward references from one element type to another; for example, edges have an attribute specifying which vertices are at each end (an attribute of type FVertexID[2]). With an indexer, it is possible to look up which edges contain a given vertex, even though this is not explicitly stored. Indexers are designed to do minimal allocations and update lazily and in batch when necessary.
- Added support for preserving UV topology in static meshes. UVs are now a first-class element type which may be indexed directly from triangles.
- Added the facility to access the underlying array in an attribute array directly.
- Triangles now directly reference their vertex, edge and UV IDs. Vertex instances are to be deprecated.
- Changed various systems to be triangle-centric rather than polygon-centric, as this is faster. Triangles are presumed to be the elementary face type in a MeshDescription, even if polygons are still supported. The concept of polygons will be somewhat shifted to mean a group of triangles which should be treated collectively for editing purposes.
- Optimised normal/tangent generation and FBX import.
- Deprecated EditableMesh, MeshEditor and StaticMeshEditorExtension plugins - these are to be removed, but they still have certain hooks in place which need removing.

diff --git a/Engine/Plugins/Editor/GLTFImporter/Source/GLTFCore/Private/GLTFMeshFactory.cpp b/Engine/Plugins/Editor/GLTFImporter/Source/GLTFCore/Private/GLTFMeshFactory.cpp
--- a/Engine/Plugins/Editor/GLTFImporter/Source/GLTFCore/Private/GLTFMeshFactory.cpp
+++ b/Engine/Plugins/Editor/GLTFImporter/Source/GLTFCore/Private/GLTFMeshFactory.cpp
@@ -173,7 +173,7 @@ namespace GLTF
 		TVertexInstanceAttributesRef<float> VertexInstanceBinormalSigns = StaticMeshAttributes.GetVertexInstanceBinormalSigns();
 		TVertexInstanceAttributesRef<FVector2D> VertexInstanceUVs = StaticMeshAttributes.GetVertexInstanceUVs();
 		TVertexInstanceAttributesRef<FVector4> VertexInstanceColors = StaticMeshAttributes.GetVertexInstanceColors();
-		VertexInstanceUVs.SetNumIndices(NumUVs);
+		VertexInstanceUVs.SetNumChannels(NumUVs);
 
 		MaterialIndicesUsed.Empty(10);
 		// Add the vertex position
@@ -394,7 +394,10 @@ namespace GLTF
 				// (Is there a way to set auto-gen smoothing threshold? glTF spec says to generate flat normals if they're not specified.
 				//   We want to combine identical verts whether they're smooth neighbors or triangles belonging to the same flat polygon.)
 				EdgeHardnesses[NewEdgeID] = false;
-				EdgeCreaseSharpnesses[NewEdgeID] = 0.0f;
+				if (EdgeCreaseSharpnesses.IsValid())
+				{
+					EdgeCreaseSharpnesses[NewEdgeID] = 0.0f;
+				}
 			}
 		}
 		return bHasDegenerateTriangles;
diff --git a/Engine/Plugins/Editor/MeshEditor/Source/MeshEditor/MeshEditorGeometryCollectionAdapter.cpp b/Engine/Plugins/Editor/MeshEditor/Source/MeshEditor/MeshEditorGeometryCollectionAdapter.cpp
--- a/Engine/Plugins/Editor/MeshEditor/Source/MeshEditor/MeshEditorGeometryCollectionAdapter.cpp
+++ b/Engine/Plugins/Editor/MeshEditor/Source/MeshEditor/MeshEditorGeometryCollectionAdapter.cpp
@@ -125,11 +125,13 @@ void UMeshEditorGeometryCollectionAdapter::AddPolygonGroupToWireframe(const FMes
 			check(GeometryCollection);
 			FColor PolygonGroupColour = FColor::Blue;
 
-			const TArray<FPolygonID>& PolygonGroupIDs = MeshDescription->GetPolygonGroupPolygons(FPolygonGroupID(SelectedPolygonGroup));
+			TArrayView<const FPolygonID> PolygonGroupIDs = MeshDescription->GetPolygonGroupPolygons(FPolygonGroupID(SelectedPolygonGroup));
 			if (PolygonGroupIDs.Num() > 0)
 			{
 				// Add all polygons and edge instances
+				PRAGMA_DISABLE_DEPRECATION_WARNINGS
 				TPolygonAttributesConstRef<FVector> PolygonNormals = MeshDescription->PolygonAttributes().GetAttributesRef<FVector>(MeshAttribute::Polygon::Normal);
+				PRAGMA_ENABLE_DEPRECATION_WARNINGS
 
 				for (const FPolygonID PolygonID : PolygonGroupIDs)
 				{
@@ -158,7 +160,9 @@ void UMeshEditorGeometryCollectionAdapter::OnRebuildRenderMeshFinish( const UEdi
 	const FMeshDescription* MeshDescription = EditableMesh->GetMeshDescription();
 	check( MeshDescription );
 
+	PRAGMA_DISABLE_DEPRECATION_WARNINGS
 	TPolygonAttributesConstRef<FVector> PolygonNormals = MeshDescription->PolygonAttributes().GetAttributesRef<FVector>( MeshAttribute::Polygon::Normal );
+	PRAGMA_ENABLE_DEPRECATION_WARNINGS
 	for( const FPolygonID PolygonID : EditableMesh->PolygonsPendingNewTangentBasis )
 	{
 		WireframeMesh->SetPolygonNormal( PolygonID, PolygonNormals[ PolygonID ] );
@@ -377,7 +381,9 @@ void UMeshEditorGeometryCollectionAdapter::OnCreatePolygons( const UEditableMesh
 	const FMeshDescription* MeshDescription = EditableMesh->GetMeshDescription();
 	check( MeshDescription );
 
+	PRAGMA_DISABLE_DEPRECATION_WARNINGS
 	TPolygonAttributesConstRef<FVector> PolygonNormals = MeshDescription->PolygonAttributes().GetAttributesRef<FVector>( MeshAttribute::Polygon::Normal );
+	PRAGMA_ENABLE_DEPRECATION_WARNINGS
 	for( const FPolygonID PolygonID : PolygonIDs )
 	{
 		WireframeMesh->AddPolygon( PolygonID );
diff --git a/Engine/Plugins/Editor/MeshEditor/Source/MeshEditor/MeshEditorMode.cpp b/Engine/Plugins/Editor/MeshEditor/Source/MeshEditor/MeshEditorMode.cpp
--- a/Engine/Plugins/Editor/MeshEditor/Source/MeshEditor/MeshEditorMode.cpp
+++ b/Engine/Plugins/Editor/MeshEditor/Source/MeshEditor/MeshEditorMode.cpp
@@ -1437,7 +1437,8 @@ void FMeshEditorMode::UpdateDebugNormals()
 			// @todo mesheditor: total debug feature for now. Need a way of making this look nice.
 			const float Length = 10.0f; // @todo mesheditor: determine length of debug line from distance from the mesh origin to the camera?
 
-			for( const FVertexInstanceID VertexInstanceID : MeshDescription->GetPolygonVertexInstances( PolygonID ) )
+			TArray<FVertexInstanceID> VertexInstanceIDs = MeshDescription->GetPolygonVertexInstances( PolygonID );
+			for( const FVertexInstanceID VertexInstanceID : VertexInstanceIDs )
 			{
 				const FVector Position = VertexPositions[ MeshDescription->GetVertexInstanceVertex( VertexInstanceID ) ];
 				const FVector Normal = VertexNormals[ VertexInstanceID ];
@@ -1549,7 +1550,7 @@ void FMeshEditorMode::UpdateSelectedElementsOverlay()
 			{
 				const FSubdividedWireEdge& SubdividedWireEdge = SubdivisionLimitData.SubdividedWireEdges[ WireEdgeNumber ];
 
-				if( SubdividedWireEdge.CounterpartEdgeID != FEdgeID::Invalid && EdgesToHighlight.Contains( SubdividedWireEdge.CounterpartEdgeID ) )
+				if( SubdividedWireEdge.CounterpartEdgeID != INDEX_NONE && EdgesToHighlight.Contains( SubdividedWireEdge.CounterpartEdgeID ) )
 				{
 					const int32 EdgeVertexIndex0 = SubdividedWireEdge.EdgeVertex0PositionIndex;
 					const int32 EdgeVertexIndex1 = SubdividedWireEdge.EdgeVertex1PositionIndex;
@@ -2110,7 +2111,8 @@ void FMeshEditorMode::FrameSelectedElements( FEditorViewportClient* ViewportClie
 					{
 						const FPolygonID PolygonID( PolygonElement.ElementAddress.ElementID );
 
-						for( const FVertexInstanceID VertexInstanceID : MeshDescription->GetPolygonVertexInstances( PolygonID ) )
+						TArray<FVertexInstanceID> VertexInstanceIDs = MeshDescription->GetPolygonVertexInstances( PolygonID );
+						for( const FVertexInstanceID VertexInstanceID : VertexInstanceIDs )
 						{
 							const FVector VertexPosition = VertexPositions[ MeshDescription->GetVertexInstanceVertex( VertexInstanceID ) ];
 							BoundingBox += Component->GetComponentTransform().TransformPosition( VertexPosition );
@@ -2229,10 +2231,10 @@ bool FMeshEditorMode::WeldSelectedVertices()
 			VertexIDsToWeld.Add( VertexID );
 		}
 
-		FVertexID WeldedVertexID = FVertexID::Invalid;
+		FVertexID WeldedVertexID = INDEX_NONE;
 		EditableMesh->WeldVertices( VertexIDsToWeld, WeldedVertexID );
 
-		if( WeldedVertexID != FVertexID::Invalid )
+		if( WeldedVertexID != INDEX_NONE )
 		{
 			FMeshElement NewVertexMeshElement;
 			{
@@ -2490,7 +2492,7 @@ void FMeshEditorMode::AddMeshElementToOverlay( UOverlayComponent* OverlayCompone
 					case EEditableMeshElementType::Polygon:
 					{
 						const FPolygonID PolygonID( MeshElement.ElementAddress.ElementID );
-						const TArray<FTriangleID>& TriangleIDs = MeshDescription->GetPolygonTriangleIDs( PolygonID );
+						TArrayView<const FTriangleID> TriangleIDs = MeshDescription->GetPolygonTriangleIDs( PolygonID );
 
 						for( int32 PolygonTriangle = 0; PolygonTriangle < TriangleIDs.Num(); PolygonTriangle++ )
 						{
@@ -3189,7 +3191,7 @@ void FMeshEditorMode::UpdateActiveAction( const bool bIsActionFinishing )
 
 					// Find first valid polygon group to add the polygon to
 					FPolygonGroupID PolygonGroupID = EditableMesh->GetFirstValidPolygonGroup();
-					check( PolygonGroupID != FPolygonGroupID::Invalid );
+					check( PolygonGroupID != INDEX_NONE );
 
 					// Create new polygon
 					PolygonsToCreate.Emplace();
@@ -3512,7 +3514,7 @@ bool FMeshEditorMode::FindEdgeSplitUnderInteractor(	UViewportInteractor* Viewpor
 {
 	check( ViewportInteractor != nullptr );
 
-	OutClosestEdgeID = FEdgeID::Invalid;
+	OutClosestEdgeID = INDEX_NONE;
 	bool bFoundSplit = false;
 
 	const FMeshDescription* MeshDescription = EditableMesh->GetMeshDescription();
@@ -4354,7 +4356,9 @@ void FMeshEditorMode::RefreshTransformables( const bool bNewObjectsSelected )
 						{
 							const FPolygonID PolygonID( MeshElement.ElementAddress.ElementID );
 
+							PRAGMA_DISABLE_DEPRECATION_WARNINGS
 							TPolygonAttributesConstRef<FVector> PolygonCenters = EditableMesh->GetMeshDescription()->PolygonAttributes().GetAttributesRef<FVector>( MeshAttribute::Polygon::Center );
+							PRAGMA_ENABLE_DEPRECATION_WARNINGS
 
 							const FVector ComponentSpacePolygonCenter = PolygonCenters[ PolygonID ];
 							ElementTransform.SetLocation( ComponentToWorldMatrix.TransformPosition( ComponentSpacePolygonCenter ) );
diff --git a/Engine/Plugins/Editor/MeshEditor/Source/MeshEditor/MeshEditorSelectionModifiers.cpp b/Engine/Plugins/Editor/MeshEditor/Source/MeshEditor/MeshEditorSelectionModifiers.cpp
--- a/Engine/Plugins/Editor/MeshEditor/Source/MeshEditor/MeshEditorSelectionModifiers.cpp
+++ b/Engine/Plugins/Editor/MeshEditor/Source/MeshEditor/MeshEditorSelectionModifiers.cpp
@@ -144,7 +144,7 @@ bool USelectPolygonsByConnectivity::ModifySelection( TMap< UEditableMesh*, TArra
 				FilledPolygons.Add( PolygonID );
 			}
 
-			FPolygonID ConnectedPolygonID = FPolygonID::Invalid;
+			FPolygonID ConnectedPolygonID = INDEX_NONE;
 
 			while ( ConnectedPolygons.Num() > 0 )
 			{
diff --git a/Engine/Plugins/Editor/MeshEditor/Source/MeshEditor/MeshEditorStaticMeshAdapter.cpp b/Engine/Plugins/Editor/MeshEditor/Source/MeshEditor/MeshEditorStaticMeshAdapter.cpp
--- a/Engine/Plugins/Editor/MeshEditor/Source/MeshEditor/MeshEditorStaticMeshAdapter.cpp
+++ b/Engine/Plugins/Editor/MeshEditor/Source/MeshEditor/MeshEditorStaticMeshAdapter.cpp
@@ -70,7 +70,9 @@ void UMeshEditorStaticMeshAdapter::OnRebuildRenderMesh( const UEditableMesh* Edi
 
 	// Add all polygons and edge instances
 	// @todo mesheditor wireframe: holes not yet handled
+	PRAGMA_DISABLE_DEPRECATION_WARNINGS
 	TPolygonAttributesConstRef<FVector> PolygonNormals = MeshDescription->PolygonAttributes().GetAttributesRef<FVector>( MeshAttribute::Polygon::Normal );
+	PRAGMA_ENABLE_DEPRECATION_WARNINGS
 	for( const FPolygonID PolygonID : EditableMesh->GetMeshDescription()->Polygons().GetElementIDs() )
 	{
 		WireframeMesh->AddPolygon( PolygonID );
@@ -91,7 +93,9 @@ void UMeshEditorStaticMeshAdapter::OnRebuildRenderMeshFinish( const UEditableMes
 	const FMeshDescription* MeshDescription = EditableMesh->GetMeshDescription();
 	check( MeshDescription );
 
+	PRAGMA_DISABLE_DEPRECATION_WARNINGS
 	TPolygonAttributesConstRef<FVector> PolygonNormals = MeshDescription->PolygonAttributes().GetAttributesRef<FVector>( MeshAttribute::Polygon::Normal );
+	PRAGMA_ENABLE_DEPRECATION_WARNINGS
 	for( const FPolygonID PolygonID : EditableMesh->PolygonsPendingNewTangentBasis )
 	{
 		WireframeMesh->SetPolygonNormal( PolygonID, PolygonNormals[ PolygonID ] );
@@ -283,7 +287,9 @@ void UMeshEditorStaticMeshAdapter::OnCreatePolygons( const UEditableMesh* Editab
 	const FMeshDescription* MeshDescription = EditableMesh->GetMeshDescription();
 	check( MeshDescription );
 
+	PRAGMA_DISABLE_DEPRECATION_WARNINGS
 	TPolygonAttributesConstRef<FVector> PolygonNormals = MeshDescription->PolygonAttributes().GetAttributesRef<FVector>( MeshAttribute::Polygon::Normal );
+	PRAGMA_ENABLE_DEPRECATION_WARNINGS
 	for( const FPolygonID PolygonID : PolygonIDs )
 	{
 		WireframeMesh->AddPolygon( PolygonID );
diff --git a/Engine/Plugins/Editor/MeshEditor/Source/MeshEditor/MeshEditorSubdividedStaticMeshAdapter.cpp b/Engine/Plugins/Editor/MeshEditor/Source/MeshEditor/MeshEditorSubdividedStaticMeshAdapter.cpp
--- a/Engine/Plugins/Editor/MeshEditor/Source/MeshEditor/MeshEditorSubdividedStaticMeshAdapter.cpp
+++ b/Engine/Plugins/Editor/MeshEditor/Source/MeshEditor/MeshEditorSubdividedStaticMeshAdapter.cpp
@@ -72,7 +72,7 @@ void UMeshEditorSubdividedStaticMeshAdapter::OnRebuildRenderMesh( const UEditabl
 			const FEdgeID EdgeID( EdgeIndex );
 			WireframeMesh->AddEdge( EdgeID );
 			WireframeMesh->SetEdgeVertices( EdgeID, FVertexID( SubdividedWireEdge.EdgeVertex0PositionIndex ), FVertexID( SubdividedWireEdge.EdgeVertex1PositionIndex ) );
-			WireframeMesh->SetEdgeColor( EdgeID, GetSubdividedEdgeColor( SubdividedWireEdge.CounterpartEdgeID != FEdgeID::Invalid ) );
+			WireframeMesh->SetEdgeColor( EdgeID, GetSubdividedEdgeColor( SubdividedWireEdge.CounterpartEdgeID != INDEX_NONE ) );
 			WireframeMesh->AddEdgeInstance( EdgeID, DummyPolygonID );
 			EdgeIndex++;
 		}
diff --git a/Engine/Plugins/Editor/MeshEditor/Source/MeshEditor/MeshEditorUtilities.cpp b/Engine/Plugins/Editor/MeshEditor/Source/MeshEditor/MeshEditorUtilities.cpp
--- a/Engine/Plugins/Editor/MeshEditor/Source/MeshEditor/MeshEditorUtilities.cpp
+++ b/Engine/Plugins/Editor/MeshEditor/Source/MeshEditor/MeshEditorUtilities.cpp
@@ -25,7 +25,7 @@ bool FMeshEditorUtilities::AssignMaterialToPolygons( UMaterialInterface* Selecte
 		// @todo mesheditor: This currently imposes the limitation that each polygon group has a unique material.
 		// Eventually we will need to be able to specify PolygonGroup properties in the editor, and ask the user for
 		// further details if there is more than one polygon group which matches the material.
-		FPolygonGroupID PolygonGroupToAssign = FPolygonGroupID::Invalid;
+		FPolygonGroupID PolygonGroupToAssign = INDEX_NONE;
 
 		if ( SelectedMaterial != nullptr )
 		{
@@ -38,7 +38,7 @@ bool FMeshEditorUtilities::AssignMaterialToPolygons( UMaterialInterface* Selecte
 				{
 					// We only expect to find one polygon group containing this material at the moment.
 					// We need to provide a way of distinguishing different polygon groups with the same material.
-					ensure( PolygonGroupToAssign == FPolygonGroupID::Invalid );
+					ensure( PolygonGroupToAssign == INDEX_NONE );
 					PolygonGroupToAssign = PolygonGroupID;
 				}
 			}
@@ -51,7 +51,7 @@ bool FMeshEditorUtilities::AssignMaterialToPolygons( UMaterialInterface* Selecte
 
 
 		// If we didn't find the material being used anywhere, create a new polygon group
-		if( PolygonGroupToAssign == FPolygonGroupID::Invalid )
+		if( PolygonGroupToAssign == INDEX_NONE )
 		{
 			// Helper function which returns a unique FName for the material slot name, based on the material's asset name,
 			// and adding a unique suffix if there are other polygon groups with the same material slot name.
diff --git a/Engine/Plugins/Editor/MeshEditor/Source/PolygonModeling/EditVertexOrEdgeSharpness.cpp b/Engine/Plugins/Editor/MeshEditor/Source/PolygonModeling/EditVertexOrEdgeSharpness.cpp
--- a/Engine/Plugins/Editor/MeshEditor/Source/PolygonModeling/EditVertexOrEdgeSharpness.cpp
+++ b/Engine/Plugins/Editor/MeshEditor/Source/PolygonModeling/EditVertexOrEdgeSharpness.cpp
@@ -112,26 +112,6 @@ void UEditVertexCornerSharpnessCommand::ApplyDuringDrag( IMeshEditorModeEditingC
 
 			static TArray<float> NewSharpnessValues;
 			NewSharpnessValues.Reset();
-
-			const TVertexAttributesRef<float> VertexSharpnesses = EditableMesh->GetMeshDescription()->VertexAttributes().GetAttributesRef<float>( MeshAttribute::Vertex::CornerSharpness );
-
-			for( const FMeshElement& VertexElement : VertexElements )
-			{
-				const FVertexID VertexID( VertexElement.ElementAddress.ElementID );
-
-				const float CurrentSharpnessValue = VertexSharpnesses[ VertexID ];
-				const float NewSharpnessValue = FMath::Clamp( CurrentSharpnessValue + ScaledDragDelta, 0.0f, 1.0f );
-
-				VertexIDs.Add( VertexID );
-				NewSharpnessValues.Add( NewSharpnessValue );
-			}
-
-
-			verify( !EditableMesh->AnyChangesToUndo() );
-
-			EditableMesh->SetVerticesCornerSharpness( VertexIDs, NewSharpnessValues );
-
-			MeshEditorMode.TrackUndo( EditableMesh, EditableMesh->MakeUndo() );
 		}
 	}
 }
@@ -164,30 +144,12 @@ void UEditEdgeCreaseSharpnessCommand::ApplyDuringDrag( IMeshEditorModeEditingCon
 			UEditableMesh* EditableMesh = MeshAndSelectedEdges.Key;
 			const TArray<FMeshElement>& EdgeElements = MeshAndSelectedEdges.Value;
 
-			const TEdgeAttributesRef<float> EdgeSharpnesses = EditableMesh->GetMeshDescription()->EdgeAttributes().GetAttributesRef<float>( MeshAttribute::Edge::CreaseSharpness );
-
 			static TArray<FEdgeID> EdgeIDs;
 			EdgeIDs.Reset();
 
 			static TArray<float> NewSharpnessValues;
 			NewSharpnessValues.Reset();
 
-			for( const FMeshElement& EdgeElement : EdgeElements )
-			{
-				const FEdgeID EdgeID( EdgeElement.ElementAddress.ElementID );
-
-				const float CurrentSharpnessValue = EdgeSharpnesses[ EdgeID ];
-				const float NewSharpnessValue = FMath::Clamp( CurrentSharpnessValue + ScaledDragDelta, 0.0f, 1.0f );
-
-				EdgeIDs.Add( EdgeID );
-				NewSharpnessValues.Add( NewSharpnessValue );
-			}
-
-			verify( !EditableMesh->AnyChangesToUndo() );
-
-			EditableMesh->SetEdgesCreaseSharpness( EdgeIDs, NewSharpnessValues );
-
-			MeshEditorMode.TrackUndo( EditableMesh, EditableMesh->MakeUndo() );
 		}
 	}
 }
diff --git a/Engine/Plugins/Editor/MeshEditor/Source/PolygonModeling/InsertEdgeLoop.cpp b/Engine/Plugins/Editor/MeshEditor/Source/PolygonModeling/InsertEdgeLoop.cpp
--- a/Engine/Plugins/Editor/MeshEditor/Source/PolygonModeling/InsertEdgeLoop.cpp
+++ b/Engine/Plugins/Editor/MeshEditor/Source/PolygonModeling/InsertEdgeLoop.cpp
@@ -42,7 +42,7 @@ void UInsertEdgeLoopCommand::ApplyDuringDrag( IMeshEditorModeEditingContract& Me
 			const TArray<FMeshElement>& EdgeElements = MeshAndEdges.Value;
 
 			// Figure out where to add the loop along the edge
-			FEdgeID ClosestEdgeID = FEdgeID::Invalid;
+			FEdgeID ClosestEdgeID = INDEX_NONE;
 			float Split = 0.0f;
 			const bool bFoundSplit = MeshEditorMode.FindEdgeSplitUnderInteractor( MeshEditorMode.GetActiveActionInteractor(), EditableMesh, EdgeElements, /* Out */ ClosestEdgeID, /* Out */ Split );
 
diff --git a/Engine/Plugins/Editor/MeshEditor/Source/PolygonModeling/Public/SplitPolygon.h b/Engine/Plugins/Editor/MeshEditor/Source/PolygonModeling/Public/SplitPolygon.h
--- a/Engine/Plugins/Editor/MeshEditor/Source/PolygonModeling/Public/SplitPolygon.h
+++ b/Engine/Plugins/Editor/MeshEditor/Source/PolygonModeling/Public/SplitPolygon.h
@@ -17,8 +17,8 @@ protected:
 	USplitPolygonCommand()
 	  : Component( nullptr ),
 		EditableMesh( nullptr ),
-		StartingEdgeID( FEdgeID::Invalid ),
-		StartingVertexID( FVertexID::Invalid ),
+		StartingEdgeID( INDEX_NONE ),
+		StartingVertexID( INDEX_NONE ),
 		EdgeSplit( 0.0f )
 	{
 		UndoText = NSLOCTEXT( "MeshEditor", "UndoSplitPolygon", "Split Polygon" );
diff --git a/Engine/Plugins/Editor/MeshEditor/Source/PolygonModeling/RemoveVertex.cpp b/Engine/Plugins/Editor/MeshEditor/Source/PolygonModeling/RemoveVertex.cpp
--- a/Engine/Plugins/Editor/MeshEditor/Source/PolygonModeling/RemoveVertex.cpp
+++ b/Engine/Plugins/Editor/MeshEditor/Source/PolygonModeling/RemoveVertex.cpp
@@ -66,7 +66,7 @@ void URemoveVertexCommand::Execute( IMeshEditorModeEditingContract& MeshEditorMo
 			const FVertexID VertexID( VertexElementToRemove.ElementAddress.ElementID );
 			{
 				bool bWasVertexRemoved = false;
-				FEdgeID NewEdgeID = FEdgeID::Invalid;
+				FEdgeID NewEdgeID = INDEX_NONE;
 
 				EditableMesh->TryToRemoveVertex( VertexID, /* Out */ bWasVertexRemoved, /* Out */ NewEdgeID );
 
diff --git a/Engine/Plugins/Editor/MeshEditor/Source/PolygonModeling/SplitEdge.cpp b/Engine/Plugins/Editor/MeshEditor/Source/PolygonModeling/SplitEdge.cpp
--- a/Engine/Plugins/Editor/MeshEditor/Source/PolygonModeling/SplitEdge.cpp
+++ b/Engine/Plugins/Editor/MeshEditor/Source/PolygonModeling/SplitEdge.cpp
@@ -40,7 +40,7 @@ void USplitEdgeCommand::ApplyDuringDrag( IMeshEditorModeEditingContract& MeshEdi
 			const TArray<FMeshElement>& EdgeElements = MeshAndEdges.Value;
 
 			// Figure out where to split
-			FEdgeID ClosestEdgeID = FEdgeID::Invalid;
+			FEdgeID ClosestEdgeID = INDEX_NONE;
 			float Split = 0.0f;
 			const bool bFoundSplit = MeshEditorMode.FindEdgeSplitUnderInteractor( ViewportInteractor, EditableMesh, EdgeElements, /* Out */ ClosestEdgeID, /* Out */ Split );
 
diff --git a/Engine/Plugins/Editor/MeshEditor/Source/PolygonModeling/SplitMesh.cpp b/Engine/Plugins/Editor/MeshEditor/Source/PolygonModeling/SplitMesh.cpp
--- a/Engine/Plugins/Editor/MeshEditor/Source/PolygonModeling/SplitMesh.cpp
+++ b/Engine/Plugins/Editor/MeshEditor/Source/PolygonModeling/SplitMesh.cpp
@@ -94,7 +94,7 @@ static void CopyAllAttributes(TAttributesSet<ElementIDType>& DestAttributesSet,
 	SrcAttributesSet.ForEach(
 		[&DestAttributesSet, ElementID](const FName AttributeName, auto AttributeArrayRef)
 		{
-			for(int32 Index = 0; Index < AttributeArrayRef.GetNumIndices(); ++Index)
+			for(int32 Index = 0; Index < AttributeArrayRef.GetNumChannels(); ++Index)
 			{
 				DestAttributesSet.SetAttribute(ElementID, AttributeName, Index, AttributeArrayRef.Get(ElementID, Index));
 			}
@@ -108,7 +108,7 @@ static void CopyAllAttributesToDifferentElement(TAttributesSet<ElementIDType>& D
 	SrcAttributesSet.ForEach(
 		[&DestAttributesSet, DstElementID, SrcElementID](const FName AttributeName, auto AttributeArrayRef)
 		{
-			for(int32 Index = 0; Index < AttributeArrayRef.GetNumIndices(); ++Index)
+			for(int32 Index = 0; Index < AttributeArrayRef.GetNumChannels(); ++Index)
 			{
 				DestAttributesSet.SetAttribute(DstElementID, AttributeName, Index, AttributeArrayRef.Get(SrcElementID, Index));
 			}
@@ -224,11 +224,11 @@ void USplitMeshCommand::Execute(IMeshEditorModeEditingContract& MeshEditorMode)
 		for(const FEdgeID EdgeId : BoundaryIds)
 		{
 			FVertexID Vertex0 = Mesh->GetMeshDescription()->GetEdgeVertex(EdgeId, 0);
-			const TArray<FVertexInstanceID>& VertexInstances0 = Mesh->GetMeshDescription()->GetVertexVertexInstances(Vertex0);
+			TArrayView<const FVertexInstanceID> VertexInstances0 = Mesh->GetMeshDescription()->GetVertexVertexInstances(Vertex0);
 			check(VertexInstances0.Num() > 0);
 
 			FVertexID Vertex1 = Mesh->GetMeshDescription()->GetEdgeVertex(EdgeId, 1);
-			const TArray<FVertexInstanceID>& VertexInstances1 = Mesh->GetMeshDescription()->GetVertexVertexInstances(Vertex1);
+			TArrayView<const FVertexInstanceID> VertexInstances1 = Mesh->GetMeshDescription()->GetVertexVertexInstances(Vertex1);
 			check(VertexInstances1.Num() > 0);
 
 			TArray<FVertexInstanceID> PolygonVertexInstances;
@@ -263,8 +263,8 @@ void USplitMeshCommand::Execute(IMeshEditorModeEditingContract& MeshEditorMode)
 
 		// @todo (mlentine): Need to make sure all numbers are the same
 		{
-			const int32 NumTexCoordIndices = Mesh->GetMeshDescription()->VertexInstanceAttributes().GetAttributeIndexCount<FVector2D>(MeshAttribute::VertexInstance::TextureCoordinate);
-			NewMeshDescription->VertexInstanceAttributes().SetAttributeIndexCount<FVector2D>(MeshAttribute::VertexInstance::TextureCoordinate, NumTexCoordIndices);
+			const int32 NumTexCoordIndices = Mesh->GetMeshDescription()->VertexInstanceAttributes().GetAttributeChannelCount(MeshAttribute::VertexInstance::TextureCoordinate);
+			NewMeshDescription->VertexInstanceAttributes().SetAttributeChannelCount(MeshAttribute::VertexInstance::TextureCoordinate, NumTexCoordIndices);
 		}
 
 		TSet<FVertexInstanceID> VertexInstanceSet;
@@ -273,7 +273,7 @@ void USplitMeshCommand::Execute(IMeshEditorModeEditingContract& MeshEditorMode)
 		TSet<FPolygonGroupID> PolygonGroupSet;
 		for(const FPolygonID PolygonId : NewPolygonIds)
 		{
-			const TArray<FVertexInstanceID>& VertexInstanceIds = Mesh->GetMeshDescription()->GetPolygonVertexInstances(PolygonId);
+			TArray<FVertexInstanceID> VertexInstanceIds = Mesh->GetMeshDescription()->GetPolygonVertexInstances(PolygonId);
 			// @todo (mlentine): Remove this when we don't need duplicate vertex instance ids.
 			TArray<FVertexInstanceID> NewVertexInstanceIds;
 			for(const FVertexInstanceID& VertexInstanceId : VertexInstanceIds)
diff --git a/Engine/Plugins/Editor/MeshEditor/Source/PolygonModeling/SplitPolygon.cpp b/Engine/Plugins/Editor/MeshEditor/Source/PolygonModeling/SplitPolygon.cpp
--- a/Engine/Plugins/Editor/MeshEditor/Source/PolygonModeling/SplitPolygon.cpp
+++ b/Engine/Plugins/Editor/MeshEditor/Source/PolygonModeling/SplitPolygon.cpp
@@ -36,8 +36,8 @@ bool USplitPolygonCommand::TryStartingToDrag( IMeshEditorModeEditingContract& Me
 {
 	Component = nullptr;
 	EditableMesh = nullptr;
-	StartingEdgeID = FEdgeID::Invalid;
-	StartingVertexID = FVertexID::Invalid;
+	StartingEdgeID = INDEX_NONE;
+	StartingVertexID = INDEX_NONE;
 	EdgeSplit = 0.0f;
 
 	// Figure out what to split
@@ -57,7 +57,7 @@ bool USplitPolygonCommand::TryStartingToDrag( IMeshEditorModeEditingContract& Me
 			const TArray<FMeshElement>& EdgeElements = MeshAndEdges.Value;
 
 			// Figure out where to split
-			FEdgeID ClosestEdgeID = FEdgeID::Invalid;
+			FEdgeID ClosestEdgeID = INDEX_NONE;
 			float Split = 0.0f;
 			const bool bFoundSplit = MeshEditorMode.FindEdgeSplitUnderInteractor( ViewportInteractor, EdgeEditableMesh, EdgeElements, /* Out */ ClosestEdgeID, /* Out */ Split );
 
@@ -117,17 +117,17 @@ void USplitPolygonCommand::ApplyDuringDrag( IMeshEditorModeEditingContract& Mesh
 		// or vertices.)
 		static TArray<FPolygonID> CandidatePolygonIDs;
 		CandidatePolygonIDs.Reset();
-		if( StartingEdgeID != FEdgeID::Invalid )
+		if( StartingEdgeID != INDEX_NONE )
 		{
 			EditableMesh->GetEdgeConnectedPolygons( StartingEdgeID, /* Out */ CandidatePolygonIDs );
 		}
-		else if( ensure( StartingVertexID != FVertexID::Invalid ) )
+		else if( ensure( StartingVertexID != INDEX_NONE ) )
 		{
 			EditableMesh->GetVertexConnectedPolygons( StartingVertexID, /* Out */ CandidatePolygonIDs );
 		}
 
-		FPolygonID PolygonToSplit = FPolygonID::Invalid;
-		FVertexID ToVertexID = FVertexID::Invalid;
+		FPolygonID PolygonToSplit = INDEX_NONE;
+		FVertexID ToVertexID = INDEX_NONE;
 
 
 		{
@@ -141,11 +141,11 @@ void USplitPolygonCommand::ApplyDuringDrag( IMeshEditorModeEditingContract& Mesh
 					const FTransform ComponentToWorld = Component->GetComponentToWorld();
 
 					FVector SplitStartLocation;
-					if( StartingVertexID != FVertexID::Invalid )
+					if( StartingVertexID != INDEX_NONE )
 					{
 						SplitStartLocation = ComponentToWorld.TransformPosition( VertexPositions[ StartingVertexID ] );
 					}
-					else if( ensure( StartingEdgeID != FEdgeID::Invalid ) )
+					else if( ensure( StartingEdgeID != INDEX_NONE ) )
 					{
 						FVertexID EdgeVertex0, EdgeVertex1;
 						EditableMesh->GetEdgeVertices( StartingEdgeID, /* Out */ EdgeVertex0, /* Out */ EdgeVertex1 );
@@ -164,13 +164,13 @@ void USplitPolygonCommand::ApplyDuringDrag( IMeshEditorModeEditingContract& Mesh
 					// @todo mesheditor urgent: Can crash with "Colinear points in FMath::ComputeBaryCentric2D()"  Needs repro.
 					FVector TriangleVertexWeights;
 					FTriangleID TriangleID = EditableMesh->ComputeBarycentricWeightForPointOnPolygon( CandidatePolygonID, ComponentToWorld.InverseTransformPosition( LaserImpactOnPolygonPlane ), /* Out */ TriangleVertexWeights );
-					if( TriangleID != FTriangleID::Invalid )
+					if( TriangleID != INDEX_NONE )
 					{
 						const FVector SplitDirection = ( LaserImpactOnPolygonPlane - SplitStartLocation ).GetSafeNormal();
 
 						// Trace out within the polygon to figure out where the split should connect to
 						float ClosestEdgeDistance = TNumericLimits<float>::Max();
-						FEdgeID ClosestEdgeID = FEdgeID::Invalid;
+						FEdgeID ClosestEdgeID = INDEX_NONE;
 
 						static TArray<FEdgeID> PolygonPerimeterEdgeIDs;
 						EditableMesh->GetPolygonPerimeterEdges( CandidatePolygonID, /* Out */ PolygonPerimeterEdgeIDs );
@@ -179,7 +179,7 @@ void USplitPolygonCommand::ApplyDuringDrag( IMeshEditorModeEditingContract& Mesh
 						{
 							bool bIsDisqualified = false;
 
-							if( StartingEdgeID != FEdgeID::Invalid )
+							if( StartingEdgeID != INDEX_NONE )
 							{
 								if( StartingEdgeID == TargetEdgeID )
 								{
@@ -187,7 +187,7 @@ void USplitPolygonCommand::ApplyDuringDrag( IMeshEditorModeEditingContract& Mesh
 									bIsDisqualified = true;
 								}
 							}
-							else if( ensure( StartingVertexID != FVertexID::Invalid ) )
+							else if( ensure( StartingVertexID != INDEX_NONE ) )
 							{
 								static TArray<FEdgeID> AdjacentEdgeIDs;
 								EditableMesh->GetVertexConnectedEdges( StartingVertexID, /* Out */ AdjacentEdgeIDs );
@@ -246,7 +246,7 @@ void USplitPolygonCommand::ApplyDuringDrag( IMeshEditorModeEditingContract& Mesh
 							}
 						}
 
-						if( ClosestEdgeID != FEdgeID::Invalid )
+						if( ClosestEdgeID != INDEX_NONE )
 						{
 							FVertexID EdgeVertex0, EdgeVertex1;
 							EditableMesh->GetEdgeVertices( ClosestEdgeID, /* Out */ EdgeVertex0, /* Out */ EdgeVertex1 );
@@ -270,7 +270,7 @@ void USplitPolygonCommand::ApplyDuringDrag( IMeshEditorModeEditingContract& Mesh
 
 								bool bIsDisqualified = false;
 
-								if( StartingEdgeID != FEdgeID::Invalid )
+								if( StartingEdgeID != INDEX_NONE )
 								{
 									FVertexID StartingEdgeVertex0, StartingEdgeVertex1;
 									EditableMesh->GetEdgeVertices( StartingEdgeID, /* Out */ StartingEdgeVertex0, /* Out */ StartingEdgeVertex1 );
@@ -281,7 +281,7 @@ void USplitPolygonCommand::ApplyDuringDrag( IMeshEditorModeEditingContract& Mesh
 										bIsDisqualified = true;
 									}
 								}
-								else if( ensure( StartingVertexID != FVertexID::Invalid ) )
+								else if( ensure( StartingVertexID != INDEX_NONE ) )
 								{
 									if( TargetVertexID == StartingVertexID )
 									{
@@ -322,7 +322,7 @@ void USplitPolygonCommand::ApplyDuringDrag( IMeshEditorModeEditingContract& Mesh
 							}
 
 							// If a vertex wasn't eligible, go ahead and connect to the edge
-							if( ToVertexID == FVertexID::Invalid )
+							if( ToVertexID == INDEX_NONE )
 							{
 								// Split the edge to create a new vertex that we'll connect to
 								PolygonToSplit = CandidatePolygonID;
@@ -348,12 +348,12 @@ void USplitPolygonCommand::ApplyDuringDrag( IMeshEditorModeEditingContract& Mesh
 		}
 
 
-		if( PolygonToSplit != FPolygonID::Invalid )
+		if( PolygonToSplit != INDEX_NONE )
 		{
-			check( ToVertexID != FVertexID::Invalid );
+			check( ToVertexID != INDEX_NONE );
 
-			FVertexID FromVertexID = FVertexID::Invalid;
-			if( StartingEdgeID != FEdgeID::Invalid )
+			FVertexID FromVertexID = INDEX_NONE;
+			if( StartingEdgeID != INDEX_NONE )
 			{
 				// First, split the edge
 				static TArray<FVertexID> NewVertexIDs;
@@ -368,11 +368,11 @@ void USplitPolygonCommand::ApplyDuringDrag( IMeshEditorModeEditingContract& Mesh
 
 				FromVertexID = NewVertexIDs[ 0 ];
 			}
-			else if( ensure( StartingVertexID != FVertexID::Invalid ) )
+			else if( ensure( StartingVertexID != INDEX_NONE ) )
 			{
 				FromVertexID = StartingVertexID;
 			}
-			check( FromVertexID != FVertexID::Invalid );
+			check( FromVertexID != INDEX_NONE );
 
 			static TArray< FPolygonToSplit > PolygonsToSplit;
 			PolygonsToSplit.Reset();
diff --git a/Engine/Plugins/Editor/MeshEditor/Source/PolygonModeling/UnifyNormals.cpp b/Engine/Plugins/Editor/MeshEditor/Source/PolygonModeling/UnifyNormals.cpp
--- a/Engine/Plugins/Editor/MeshEditor/Source/PolygonModeling/UnifyNormals.cpp
+++ b/Engine/Plugins/Editor/MeshEditor/Source/PolygonModeling/UnifyNormals.cpp
@@ -28,7 +28,7 @@ namespace UnifyNormalsCommandUtils
 		// Iterate through the polygon vertices to find vertex 0 of the MeshEdge
 		const FVertexID EdgeVertex0 = MeshDescription->GetEdgeVertex(EdgeID, 0);
 		const FVertexID EdgeVertex1 = MeshDescription->GetEdgeVertex(EdgeID, 1);
-		const TArray<FVertexInstanceID>& VertexInstances = MeshDescription->GetPolygonVertexInstances(PolygonID);
+		TArray<FVertexInstanceID> VertexInstances = MeshDescription->GetPolygonVertexInstances(PolygonID);
 		int32 NumVertices = VertexInstances.Num();
 		for (int32 Index = 0; Index < NumVertices; ++Index)
 		{
diff --git a/Engine/Plugins/Editor/SpeedTreeImporter/Source/SpeedTreeImporter/Private/SpeedTreeImportFactory.cpp b/Engine/Plugins/Editor/SpeedTreeImporter/Source/SpeedTreeImporter/Private/SpeedTreeImportFactory.cpp
--- a/Engine/Plugins/Editor/SpeedTreeImporter/Source/SpeedTreeImporter/Private/SpeedTreeImportFactory.cpp
+++ b/Engine/Plugins/Editor/SpeedTreeImporter/Source/SpeedTreeImporter/Private/SpeedTreeImportFactory.cpp
@@ -1425,7 +1425,7 @@ FVertexInstanceID ProcessTriangleCorner(
 	TVertexInstanceAttributesRef<FVector2D> VertexInstanceUVs)
 {
 	//Speedtree uses 7 or 8 UVs to store is data
-	check(VertexInstanceUVs.GetNumIndices() >= 7);
+	check(VertexInstanceUVs.GetNumChannels() >= 7);
 
 	SpeedTree::st_float32 Data[ 4 ];
 
@@ -1518,7 +1518,7 @@ FVertexInstanceID ProcessTriangleCorner(
 	}
 	else if( RenderState->m_bLeavesPresent || RenderState->m_bFacingLeavesPresent )
 	{
-		check(VertexInstanceUVs.GetNumIndices() == 8);
+		check(VertexInstanceUVs.GetNumChannels() == 8);
 
 		// anchor
 		if( RenderState->m_bFacingLeavesPresent )
@@ -1716,7 +1716,6 @@ UObject* USpeedTreeImportFactory::FactoryCreateBinary7(UClass* InClass, UObject*
 						
 						TVertexAttributesRef<FVector> VertexPositions = Attributes.GetVertexPositions();
 						TEdgeAttributesRef<bool> EdgeHardnesses = Attributes.GetEdgeHardnesses();
-						TEdgeAttributesRef<float> EdgeCreaseSharpnesses = Attributes.GetEdgeCreaseSharpnesses();
 						TPolygonGroupAttributesRef<FName> PolygonGroupImportedMaterialSlotNames = Attributes.GetPolygonGroupMaterialSlotNames();
 						TVertexInstanceAttributesRef<FVector> VertexInstanceNormals = Attributes.GetVertexInstanceNormals();
 						TVertexInstanceAttributesRef<FVector> VertexInstanceTangents = Attributes.GetVertexInstanceTangents();
@@ -1744,7 +1743,7 @@ UObject* USpeedTreeImportFactory::FactoryCreateBinary7(UClass* InClass, UObject*
 							}
 						}
 						//Speedtree use UVs to store is data
-						VertexInstanceUVs.SetNumIndices(NumUVs);
+						VertexInstanceUVs.SetNumChannels(NumUVs);
 
 						for (int32 DrawCallIndex = 0; DrawCallIndex < TreeLOD->m_nNumDrawCalls; ++DrawCallIndex)
 						{
@@ -1891,7 +1890,6 @@ UObject* USpeedTreeImportFactory::FactoryCreateBinary7(UClass* InClass, UObject*
 
 					TVertexAttributesRef<FVector> VertexPositions = Attributes.GetVertexPositions();
 					TEdgeAttributesRef<bool> EdgeHardnesses = Attributes.GetEdgeHardnesses();
-					TEdgeAttributesRef<float> EdgeCreaseSharpnesses = Attributes.GetEdgeCreaseSharpnesses();
 					TPolygonGroupAttributesRef<FName> PolygonGroupImportedMaterialSlotNames = Attributes.GetPolygonGroupMaterialSlotNames();
 					TVertexInstanceAttributesRef<FVector> VertexInstanceNormals = Attributes.GetVertexInstanceNormals();
 					TVertexInstanceAttributesRef<FVector> VertexInstanceTangents = Attributes.GetVertexInstanceTangents();
@@ -1900,7 +1898,7 @@ UObject* USpeedTreeImportFactory::FactoryCreateBinary7(UClass* InClass, UObject*
 					TVertexInstanceAttributesRef<FVector2D> VertexInstanceUVs = Attributes.GetVertexInstanceUVs();
 
 					//Speedtree use UVs to store is data
-					VertexInstanceUVs.SetNumIndices(2);
+					VertexInstanceUVs.SetNumChannels(2);
 
 					for (int32 MatIndex = 0; MatIndex < StaticMesh->StaticMaterials.Num(); ++MatIndex)
 					{
@@ -2223,7 +2221,6 @@ UObject* USpeedTreeImportFactory::FactoryCreateBinary8(UClass* InClass, UObject*
 
 			TVertexAttributesRef<FVector> VertexPositions = Attributes.GetVertexPositions();
 			TEdgeAttributesRef<bool> EdgeHardnesses = Attributes.GetEdgeHardnesses();
-			TEdgeAttributesRef<float> EdgeCreaseSharpnesses = Attributes.GetEdgeCreaseSharpnesses();
 			TPolygonGroupAttributesRef<FName> PolygonGroupImportedMaterialSlotNames = Attributes.GetPolygonGroupMaterialSlotNames();
 			TVertexInstanceAttributesRef<FVector> VertexInstanceNormals = Attributes.GetVertexInstanceNormals();
 			TVertexInstanceAttributesRef<FVector> VertexInstanceTangents = Attributes.GetVertexInstanceTangents();
@@ -2232,7 +2229,7 @@ UObject* USpeedTreeImportFactory::FactoryCreateBinary8(UClass* InClass, UObject*
 			TVertexInstanceAttributesRef<FVector2D> VertexInstanceUVs = Attributes.GetVertexInstanceUVs();
 
 			//Speedtree use 8 UVs to store is data
-			VertexInstanceUVs.SetNumIndices(8);
+			VertexInstanceUVs.SetNumChannels(8);
 
 			for (int32 MatIndex = 0; MatIndex < StaticMesh->StaticMaterials.Num(); ++MatIndex)
 			{
diff --git a/Engine/Plugins/Enterprise/DataprepEditor/Source/DataprepCore/Private/DataprepCorePrivateUtils.cpp b/Engine/Plugins/Enterprise/DataprepEditor/Source/DataprepCore/Private/DataprepCorePrivateUtils.cpp
--- a/Engine/Plugins/Enterprise/DataprepEditor/Source/DataprepCore/Private/DataprepCorePrivateUtils.cpp
+++ b/Engine/Plugins/Enterprise/DataprepEditor/Source/DataprepCore/Private/DataprepCorePrivateUtils.cpp
@@ -119,14 +119,14 @@ void DataprepCorePrivateUtils::BuildStaticMeshes(TSet<UStaticMesh*>& StaticMeshe
 					{
 						TVertexInstanceAttributesConstRef<FVector2D> VertexInstanceUVs = Attributes.GetVertexInstanceUVs();
 						//If the importer have enabled lightmap generation, disabling it may interfere with the build process, so we only allow enabling it.
-						SourceModel.BuildSettings.bGenerateLightmapUVs |= VertexInstanceUVs.IsValid() && VertexInstanceUVs.GetNumIndices() > SourceModel.BuildSettings.SrcLightmapIndex;
+						SourceModel.BuildSettings.bGenerateLightmapUVs |= VertexInstanceUVs.IsValid() && VertexInstanceUVs.GetNumChannels() > SourceModel.BuildSettings.SrcLightmapIndex;
 					}
 					else
 					{
 						SourceModel.BuildSettings.bGenerateLightmapUVs = false;
 					}
 
-					SourceModel.BuildSettings.bRecomputeNormals = !(Attributes.GetVertexInstanceNormals().IsValid() && Attributes.GetVertexInstanceNormals().GetNumIndices() > 0);
+					SourceModel.BuildSettings.bRecomputeNormals = !(Attributes.GetVertexInstanceNormals().IsValid() && Attributes.GetVertexInstanceNormals().GetNumChannels() > 0);
 					SourceModel.BuildSettings.bRecomputeTangents = false;
 					//SourceModel.BuildSettings.bBuildAdjacencyBuffer = false;
 					//SourceModel.BuildSettings.bBuildReversedIndexBuffer = false;
diff --git a/Engine/Plugins/Enterprise/DataprepEditor/Source/DataprepLibraries/Private/DataprepEditingOperations.cpp b/Engine/Plugins/Enterprise/DataprepEditor/Source/DataprepLibraries/Private/DataprepEditingOperations.cpp
--- a/Engine/Plugins/Enterprise/DataprepEditor/Source/DataprepLibraries/Private/DataprepEditingOperations.cpp
+++ b/Engine/Plugins/Enterprise/DataprepEditor/Source/DataprepLibraries/Private/DataprepEditingOperations.cpp
@@ -370,8 +370,8 @@ void UDataprepCreateProxyMeshOperation::OnExecution_Implementation(const FDatapr
 		{
 			FStaticMeshAttributes Attributes(*MeshDescription);
 			bool bHasValidLightmapUVs = Attributes.GetVertexInstanceUVs().IsValid() &&
-				Attributes.GetVertexInstanceUVs().GetNumIndices() > BuildSettings.SrcLightmapIndex &&
-				Attributes.GetVertexInstanceUVs().GetNumIndices() > BuildSettings.DstLightmapIndex;
+				Attributes.GetVertexInstanceUVs().GetNumChannels() > BuildSettings.SrcLightmapIndex &&
+				Attributes.GetVertexInstanceUVs().GetNumChannels() > BuildSettings.DstLightmapIndex;
 			if(!bHasValidLightmapUVs)
 			{
 				ProxySettings.bReuseMeshLightmapUVs = false;
diff --git a/Engine/Plugins/Enterprise/DataprepEditor/Source/DataprepLibraries/Private/DataprepOperationsLibrary.cpp b/Engine/Plugins/Enterprise/DataprepEditor/Source/DataprepLibraries/Private/DataprepOperationsLibrary.cpp
--- a/Engine/Plugins/Enterprise/DataprepEditor/Source/DataprepLibraries/Private/DataprepOperationsLibrary.cpp
+++ b/Engine/Plugins/Enterprise/DataprepEditor/Source/DataprepLibraries/Private/DataprepOperationsLibrary.cpp
@@ -329,14 +329,14 @@ namespace DataprepOperationsLibraryUtil
 						if(SourceModel.BuildSettings.DstLightmapIndex != -1)
 						{
 							TVertexInstanceAttributesConstRef<FVector2D> VertexInstanceUVs = Attributes.GetVertexInstanceUVs();
-							SourceModel.BuildSettings.bGenerateLightmapUVs = VertexInstanceUVs.IsValid() && VertexInstanceUVs.GetNumIndices() > SourceModel.BuildSettings.DstLightmapIndex;
+							SourceModel.BuildSettings.bGenerateLightmapUVs = VertexInstanceUVs.IsValid() && VertexInstanceUVs.GetNumChannels() > SourceModel.BuildSettings.DstLightmapIndex;
 						}
 						else
 						{
 							SourceModel.BuildSettings.bGenerateLightmapUVs = false;
 						}
 
-						SourceModel.BuildSettings.bRecomputeNormals = !(Attributes.GetVertexInstanceNormals().IsValid() && Attributes.GetVertexInstanceNormals().GetNumIndices() > 0);
+						SourceModel.BuildSettings.bRecomputeNormals = !(Attributes.GetVertexInstanceNormals().IsValid() && Attributes.GetVertexInstanceNormals().GetNumChannels() > 0);
 						SourceModel.BuildSettings.bRecomputeTangents = false;
 						//SourceModel.BuildSettings.bBuildAdjacencyBuffer = false;
 						//SourceModel.BuildSettings.bBuildReversedIndexBuffer = false;
diff --git a/Engine/Plugins/Enterprise/DatasmithC4DImporter/Source/DatasmithC4DTranslator/Private/DatasmithC4DImporter.cpp b/Engine/Plugins/Enterprise/DatasmithC4DImporter/Source/DatasmithC4DTranslator/Private/DatasmithC4DImporter.cpp
--- a/Engine/Plugins/Enterprise/DatasmithC4DImporter/Source/DatasmithC4DTranslator/Private/DatasmithC4DImporter.cpp
+++ b/Engine/Plugins/Enterprise/DatasmithC4DImporter/Source/DatasmithC4DTranslator/Private/DatasmithC4DImporter.cpp
@@ -2756,7 +2756,7 @@ TSharedPtr<IDatasmithMeshElement> FDatasmithC4DImporter::ImportMesh(melange::Pol
 
 	// At least one UV set must exist.
 	int32 UVChannelCount = UVWTagsData.Num();
-	VertexInstanceUVs.SetNumIndices(FMath::Max(1, UVChannelCount));
+	VertexInstanceUVs.SetNumChannels(FMath::Max(1, UVChannelCount));
 
 	// Vertices
 	for (int32 PointIndex = 0; PointIndex < PointCount; ++PointIndex)
diff --git a/Engine/Plugins/Enterprise/DatasmithCADImporter/Source/CADLibrary/Private/CoreTechHelper.cpp b/Engine/Plugins/Enterprise/DatasmithCADImporter/Source/CADLibrary/Private/CoreTechHelper.cpp
--- a/Engine/Plugins/Enterprise/DatasmithCADImporter/Source/CADLibrary/Private/CoreTechHelper.cpp
+++ b/Engine/Plugins/Enterprise/DatasmithCADImporter/Source/CADLibrary/Private/CoreTechHelper.cpp
@@ -43,7 +43,7 @@ struct FVertexData
 		Index = InIndex;
 		Coordinates = V;
 		bIsMerged = false;
-		VertexID = FVertexID::Invalid;
+		VertexID = INDEX_NONE;
 	}
 };
 
@@ -258,7 +258,7 @@ bool FillMesh(const FMeshParameters& MeshParameters, const FImportParameters& Im
 	const int32 TriangleCount = 3;
 	const TriangleIndex Clockwise = { 0, 1, 2 };
 	const TriangleIndex CounterClockwise = { 0, 2, 1 };
-    const int32 InvalidID = FVertexID::Invalid.GetValue();
+    const int32 InvalidID = INDEX_NONE;
 
 
 	TArray<FVertexInstanceID> TriangleVertexInstanceIDs;
@@ -286,7 +286,7 @@ bool FillMesh(const FMeshParameters& MeshParameters, const FImportParameters& Im
 	for (const FTessellationData& FaceTessellation : FaceTessellations)
 	{
 		// we assume that face has only color
-		MaterialToPolygonGroupMapping.Add(FaceTessellation.ColorName, FPolygonGroupID::Invalid);
+		MaterialToPolygonGroupMapping.Add(FaceTessellation.ColorName, INDEX_NONE);
 	}
 
 	// Add to the mesh, a polygon groups per material
@@ -300,7 +300,7 @@ bool FillMesh(const FMeshParameters& MeshParameters, const FImportParameters& Im
 		Material.Value = PolyGroupID;
 	}
 
-	VertexInstanceUVs.SetNumIndices(1);
+	VertexInstanceUVs.SetNumChannels(1);
 
 	int32 NbStep = 1;
 	if (MeshParameters.bIsSymmetric)
@@ -454,7 +454,6 @@ bool FillMesh(const FMeshParameters& MeshParameters, const FImportParameters& Im
 bool ConvertCTBodySetToMeshDescription(const FImportParameters& ImportParams, const FMeshParameters& MeshParameters, FBodyMesh& Body, FMeshDescription& MeshDescription)
 {
 	// Ref. CreateMesh(UDatasmithCADImportOptions* CADOptions, CTMesh& Mesh)
-	MeshDescription.EdgeAttributes().RegisterAttribute<bool>(MeshAttribute::Edge::IsUVSeam, 1, false);
 
 	// in a closed big mesh VertexCount ~ TriangleCount / 2, EdgeCount ~ 1.5* TriangleCount
 	MeshDescription.ReserveNewVertexInstances(Body.TriangleCount*3);
diff --git a/Engine/Plugins/Enterprise/DatasmithCADImporter/Source/CADLibrary/Private/MeshEditingWrapper.cpp b/Engine/Plugins/Enterprise/DatasmithCADImporter/Source/CADLibrary/Private/MeshEditingWrapper.cpp
--- a/Engine/Plugins/Enterprise/DatasmithCADImporter/Source/CADLibrary/Private/MeshEditingWrapper.cpp
+++ b/Engine/Plugins/Enterprise/DatasmithCADImporter/Source/CADLibrary/Private/MeshEditingWrapper.cpp
@@ -36,12 +36,12 @@ void FMeshEditingWrapper::UpdateMeshWrapper()
 
 void FMeshEditingWrapper::DefineEdgeTopology(FEdgeID EdgeID)
 {
-	if (EdgeID == FEdgeID::Invalid)
+	if (EdgeID == INDEX_NONE)
 	{
 		return;
 	}
 
-	const TArray<FTriangleID>& EdgeConnectedPolygons = MeshDescription.GetEdgeConnectedTriangles(EdgeID);
+	TArrayView<const FTriangleID> EdgeConnectedPolygons = MeshDescription.GetEdgeConnectedTriangles(EdgeID);
 
 	switch (EdgeConnectedPolygons.Num())
 	{
@@ -62,12 +62,12 @@ void FMeshEditingWrapper::DefineEdgeTopology(FEdgeID EdgeID)
 
 void FMeshEditingWrapper::DefineVertexTopologyApproximation(FVertexID VertexID)
 {
-	if (VertexID == FVertexID::Invalid)
+	if (VertexID == INDEX_NONE)
 	{
 		return;
 	}
 
-	const TArray<FEdgeID>& VertexConnectedEdgeIDs = MeshDescription.GetVertexConnectedEdges(VertexID);
+	TArrayView<const FEdgeID> VertexConnectedEdgeIDs = MeshDescription.GetVertexConnectedEdges(VertexID);
 
 	switch (VertexConnectedEdgeIDs.Num())
 	{
@@ -113,12 +113,12 @@ void FMeshEditingWrapper::DefineVertexTopologyApproximation(FVertexID VertexID)
 			const FEdgeID FirstEdgeID(VertexConnectedEdgeIDs[0]);
 			FEdgeID EdgeID(FirstEdgeID);
 
-			FTriangleID TriangleID(FPolygonID::Invalid);
+			FTriangleID TriangleID(INDEX_NONE);
 			int32 TriangleCount = 0;
 
 			do
 			{
-				const TArray<FTriangleID>& EdgeConnectedPolygons = MeshDescription.GetEdgeConnectedTriangles(EdgeID);
+				TArrayView<const FTriangleID> EdgeConnectedPolygons = MeshDescription.GetEdgeConnectedTriangles(EdgeID);
 				// Border edge no more triangles to process, exit from the loop
 				if (EdgeConnectedPolygons.Num() < 2)
 				{
@@ -128,8 +128,7 @@ void FMeshEditingWrapper::DefineVertexTopologyApproximation(FVertexID VertexID)
 				TriangleID = (TriangleID == EdgeConnectedPolygons[0]) ? EdgeConnectedPolygons[1] : EdgeConnectedPolygons[0];
 				++TriangleCount;
 
-				TStaticArray<FEdgeID, 3> TriangleEdges;
-				MeshDescription.GetTriangleEdges(TriangleID, TriangleEdges);
+				TArrayView<const FEdgeID> TriangleEdges = MeshDescription.GetTriangleEdges(TriangleID);
 
 				for (int32 Corner = 0; Corner < 3; ++Corner)
 				{
@@ -166,8 +165,7 @@ FMeshEditingWrapper::~FMeshEditingWrapper()
 void FMeshEditingWrapper::ResetTriangleMarkerRecursively(FTriangleID Triangle)
 {
 	MeshEditingWrapperUtils::ResetElementMarker(TriangleMetaData[Triangle.GetValue()], ELEMENT_FIRST_MARKER_MASK);
-	TStaticArray<FEdgeID, 3> EdgeSet;
-	MeshDescription.GetTriangleEdges(Triangle, EdgeSet);
+	TArrayView<const FEdgeID> EdgeSet = MeshDescription.GetTriangleEdges(Triangle);
 	for (int32 IEdge = 0; IEdge < 3; IEdge++)
 	{
 		MeshEditingWrapperUtils::ResetElementMarker(TriangleMetaData[EdgeSet[IEdge].GetValue()], ELEMENT_FIRST_MARKER_MASK);
@@ -250,35 +248,30 @@ void FMeshEditingWrapper::GetVertexBoundingBox(FVertexInstanceID VertexInstanceI
 
 void FMeshEditingWrapper::SwapTriangleOrientation(FTriangleID Triangle)
 {
-	FTriangleArray& TraingleSet = MeshDescription.Triangles();
-	FVertexInstanceID Vertex0ID = TraingleSet[FTriangleID(Triangle)].GetVertexInstanceID(0);
-	FVertexInstanceID Vertex1ID = TraingleSet[FTriangleID(Triangle)].GetVertexInstanceID(1);
-	TraingleSet[Triangle].SetVertexInstanceID(0, Vertex1ID);
-	TraingleSet[Triangle].SetVertexInstanceID(1, Vertex0ID);
-
+	MeshDescription.ReverseTriangleFacing(Triangle);
+	TArrayView<const FVertexInstanceID> TriVertexInstances = MeshDescription.GetTriangleVertexInstances(Triangle);
 	for (int32 IVertex = 0; IVertex < 3; IVertex++)
 	{
-		SwapVertexNormal(TraingleSet[FTriangleID(Triangle)].GetVertexInstanceID(IVertex));
+		SwapVertexNormal(TriVertexInstances[IVertex]);
 	}
 }
 
 // Edge
 const FTriangleID FMeshEditingWrapper::GetOtherTriangleAtEdge(FEdgeID EdgeID, FTriangleID Triangle) const
 {
-	const TArray<FTriangleID>& EdgeConnectedPolygons = MeshDescription.GetEdgeConnectedTriangles(EdgeID);
-	return EdgeConnectedPolygons.Num() < 2 ? FTriangleID::Invalid : (EdgeConnectedPolygons[0] == Triangle ? EdgeConnectedPolygons[1] : EdgeConnectedPolygons[0]);
+	TArrayView<const FTriangleID> EdgeConnectedPolygons = MeshDescription.GetEdgeConnectedTriangles(EdgeID);
+	return EdgeConnectedPolygons.Num() < 2 ? INDEX_NONE : (EdgeConnectedPolygons[0] == Triangle ? EdgeConnectedPolygons[1] : EdgeConnectedPolygons[0]);
 }
 
 bool FMeshEditingWrapper::GetEdgeDirectionInTriangle(FEdgeID EdgeID, int32 TriangleIndex) const
 {
-	const TArray<FTriangleID>& EdgeConnectedTriangles = MeshDescription.GetEdgeConnectedTriangles(EdgeID);
+	TArrayView<const FTriangleID> EdgeConnectedTriangles = MeshDescription.GetEdgeConnectedTriangles(EdgeID);
 
 	if (EdgeConnectedTriangles.Num() > TriangleIndex)
 	{
 		const FTriangleID TriangleID = EdgeConnectedTriangles[TriangleIndex];
-		const TArrayView<const FVertexInstanceID>& VertexInstanceIDs = MeshDescription.GetTriangleVertexInstances(TriangleID);
-
-		TStaticArray<FEdgeID, 3> TriangleEdges = MeshDescription.GetTriangleEdges(TriangleID);
+		TArrayView<const FVertexInstanceID> VertexInstanceIDs = MeshDescription.GetTriangleVertexInstances(TriangleID);
+		TArrayView<const FEdgeID> TriangleEdges = MeshDescription.GetTriangleEdges(TriangleID);
 
 		for (int32 Corner = 0; Corner < 3; ++Corner)
 		{
diff --git a/Engine/Plugins/Enterprise/DatasmithCADImporter/Source/CADLibrary/Private/MeshOperator.cpp b/Engine/Plugins/Enterprise/DatasmithCADImporter/Source/CADLibrary/Private/MeshOperator.cpp
--- a/Engine/Plugins/Enterprise/DatasmithCADImporter/Source/CADLibrary/Private/MeshOperator.cpp
+++ b/Engine/Plugins/Enterprise/DatasmithCADImporter/Source/CADLibrary/Private/MeshOperator.cpp
@@ -19,7 +19,6 @@ bool MeshOperator::OrientMesh(FMeshDescription& MeshDescription)
 
 	TQueue<FTriangleID> Front; 
 	FTriangleID AdjacentTriangle;
-	TStaticArray<FEdgeID, 3> EdgeSet;
 	FEdgeID Edge;
 
 	TArrayView < const FVertexInstanceID> Vertices;
@@ -42,12 +41,12 @@ bool MeshOperator::OrientMesh(FMeshDescription& MeshDescription)
 			continue;
 		}
 
-		HighestVertex[0] = FVertexInstanceID::Invalid;
-		HighestVertex[1] = FVertexInstanceID::Invalid;
-		HighestVertex[2] = FVertexInstanceID::Invalid;
-		LowestVertex[0] = FVertexInstanceID::Invalid;
-		LowestVertex[1] = FVertexInstanceID::Invalid;
-		LowestVertex[2] = FVertexInstanceID::Invalid;
+		HighestVertex[0] = INDEX_NONE;
+		HighestVertex[1] = INDEX_NONE;
+		HighestVertex[2] = INDEX_NONE;
+		LowestVertex[0] = INDEX_NONE;
+		LowestVertex[1] = INDEX_NONE;
+		LowestVertex[2] = INDEX_NONE;
 
 		MaxCorner[0] = -MAX_flt;
 		MaxCorner[1] = -MAX_flt;
@@ -70,7 +69,7 @@ bool MeshOperator::OrientMesh(FMeshDescription& MeshDescription)
 		{
 			Front.Dequeue(Triangle);
 
-			MeshDescription.GetTriangleEdges(Triangle, EdgeSet);
+			TArrayView<const FEdgeID> EdgeSet = MeshDescription.GetTriangleEdges(Triangle);
 
 			for (int32 IEdge = 0; IEdge < 3; IEdge++)
 			{
@@ -114,20 +113,21 @@ bool MeshOperator::OrientMesh(FMeshDescription& MeshDescription)
 			// The normal most parallel to the axis is preferred
 			// To avoid mistake, we check for each highest vertex and trust the majority 
 
-			if (HighestVertex[0] != FVertexInstanceID::Invalid)
+			if (HighestVertex[0] != INDEX_NONE)
 			{
-				FStaticMeshAttributes StaticMeshAttributes( MeshDescription );
+				FStaticMeshConstAttributes StaticMeshAttributes( MeshDescription );
+				TVertexInstanceAttributesConstRef<FVector> Normals = StaticMeshAttributes.GetVertexInstanceNormals();
 
 				for (int32 VertexIndex = 0; VertexIndex < 3; VertexIndex++)
 				{
 					if (MeshWrapper.IsVertexOfCategory(HighestVertex[VertexIndex], EElementCategory::ElementCategorySurface))
 					{
 						FVertexID VertexID = MeshDescription.GetVertexInstanceVertex(HighestVertex[VertexIndex]);
-						const TArray<FVertexInstanceID>& CoincidentVertexInstanceIdSet = MeshDescription.GetVertexVertexInstances(VertexID);
+						TArrayView<const FVertexInstanceID> CoincidentVertexInstanceIdSet = MeshDescription.GetVertexVertexInstances(VertexID);
 						float MaxComponent = 0;
-						for (const FVertexInstanceID& VertexInstanceID : CoincidentVertexInstanceIdSet)
+						for (const FVertexInstanceID VertexInstanceID : CoincidentVertexInstanceIdSet)
 						{
-							FVector Normal = StaticMeshAttributes.GetVertexInstanceNormals()[ VertexInstanceID ];
+							FVector Normal = Normals[ VertexInstanceID ];
 							if (FMath::Abs(MaxComponent) < FMath::Abs(Normal[VertexIndex]))
 							{
 								MaxComponent = Normal[VertexIndex];
@@ -147,11 +147,11 @@ bool MeshOperator::OrientMesh(FMeshDescription& MeshDescription)
 					if (MeshWrapper.IsVertexOfCategory(LowestVertex[VertexIndex], EElementCategory::ElementCategorySurface))
 					{
 						FVertexID VertexID = MeshDescription.GetVertexInstanceVertex(LowestVertex[VertexIndex]);
-						const TArray<FVertexInstanceID>& CoincidentVertexInstanceIdSet = MeshDescription.GetVertexVertexInstances(VertexID);
+						TArrayView<const FVertexInstanceID> CoincidentVertexInstanceIdSet = MeshDescription.GetVertexVertexInstances(VertexID);
 						float MaxComponent = 0;
-						for (const FVertexInstanceID& VertexInstanceID : CoincidentVertexInstanceIdSet)
+						for (const FVertexInstanceID VertexInstanceID : CoincidentVertexInstanceIdSet)
 						{
-							FVector Normal = MeshDescription.VertexInstanceAttributes().GetAttribute<FVector>(VertexInstanceID, MeshAttribute::VertexInstance::Normal);
+							FVector Normal = Normals[ VertexInstanceID ];
 							if (FMath::Abs(MaxComponent) < FMath::Abs(Normal[VertexIndex]))
 							{
 								MaxComponent = Normal[VertexIndex];
diff --git a/Engine/Plugins/Enterprise/DatasmithCADImporter/Source/DatasmithCoreTechExtension/Private/CoreTechBlueprintLibrary.cpp b/Engine/Plugins/Enterprise/DatasmithCADImporter/Source/DatasmithCoreTechExtension/Private/CoreTechBlueprintLibrary.cpp
--- a/Engine/Plugins/Enterprise/DatasmithCADImporter/Source/DatasmithCoreTechExtension/Private/CoreTechBlueprintLibrary.cpp
+++ b/Engine/Plugins/Enterprise/DatasmithCADImporter/Source/DatasmithCoreTechExtension/Private/CoreTechBlueprintLibrary.cpp
@@ -47,10 +47,10 @@ bool UCoreTechBlueprintLibrary::RetessellateStaticMeshWithNotification(UStaticMe
 				StaticMesh->PreEditChange( nullptr );
 			}
 
-			const int32 OldNumberOfUVChannels = DestinationMeshDescription->VertexInstanceAttributes().GetAttributeIndexCount<FVector2D>(MeshAttribute::VertexInstance::TextureCoordinate);
+			const int32 OldNumberOfUVChannels = DestinationMeshDescription->VertexInstanceAttributes().GetAttributeChannelCount(MeshAttribute::VertexInstance::TextureCoordinate);
 			if (FCoreTechRetessellate_Impl::ApplyOnOneAsset( *StaticMesh, *CoreTechData, TessellationSettings ))
 			{
-				const int32 NumberOfUVChannels = DestinationMeshDescription->VertexInstanceAttributes().GetAttributeIndexCount<FVector2D>(MeshAttribute::VertexInstance::TextureCoordinate);
+				const int32 NumberOfUVChannels = DestinationMeshDescription->VertexInstanceAttributes().GetAttributeChannelCount(MeshAttribute::VertexInstance::TextureCoordinate);
 				if (NumberOfUVChannels < OldNumberOfUVChannels)
 				{
 					FailureReason = FText::Format(NSLOCTEXT("BlueprintRetessellation", "UVChannelsDestroyed", "Tessellation operation on Static Mesh {0} is destroying all UV channels above channel #{1}"), FText::FromString(StaticMesh->GetName()), NumberOfUVChannels - 1);
diff --git a/Engine/Plugins/Enterprise/DatasmithCADImporter/Source/DatasmithOpenNurbsTranslator/Private/DatasmithOpenNurbsTranslator.cpp b/Engine/Plugins/Enterprise/DatasmithCADImporter/Source/DatasmithOpenNurbsTranslator/Private/DatasmithOpenNurbsTranslator.cpp
--- a/Engine/Plugins/Enterprise/DatasmithCADImporter/Source/DatasmithOpenNurbsTranslator/Private/DatasmithOpenNurbsTranslator.cpp
+++ b/Engine/Plugins/Enterprise/DatasmithCADImporter/Source/DatasmithOpenNurbsTranslator/Private/DatasmithOpenNurbsTranslator.cpp
@@ -273,11 +273,12 @@ namespace DatasmithOpenNurbsTranslatorUtils
 		}
 
 		MeshDescription.Empty();
+		FStaticMeshAttributes Attributes(MeshDescription);
 
-		TVertexAttributesRef<FVector> VertexPositions = MeshDescription.VertexAttributes().GetAttributesRef<FVector>(MeshAttribute::Vertex::Position);
-		TVertexInstanceAttributesRef<FVector> VertexInstanceNormals = MeshDescription.VertexInstanceAttributes().GetAttributesRef<FVector>(MeshAttribute::VertexInstance::Normal);
-		TVertexInstanceAttributesRef<FVector2D> VertexInstanceUVs = MeshDescription.VertexInstanceAttributes().GetAttributesRef<FVector2D>(MeshAttribute::VertexInstance::TextureCoordinate);
-		TPolygonGroupAttributesRef<FName> PolygonGroupImportedMaterialSlotNames = MeshDescription.PolygonGroupAttributes().GetAttributesRef<FName>(MeshAttribute::PolygonGroup::ImportedMaterialSlotName);
+		TVertexAttributesRef<FVector> VertexPositions = Attributes.GetVertexPositions();
+		TVertexInstanceAttributesRef<FVector> VertexInstanceNormals = Attributes.GetVertexInstanceNormals();
+		TVertexInstanceAttributesRef<FVector2D> VertexInstanceUVs = Attributes.GetVertexInstanceUVs();
+		TPolygonGroupAttributesRef<FName> PolygonGroupImportedMaterialSlotNames = Attributes.GetPolygonGroupMaterialSlotNames();
 
 		// Prepared for static mesh usage ?
 		if (!VertexPositions.IsValid() || !VertexInstanceNormals.IsValid() || !VertexInstanceUVs.IsValid() || !PolygonGroupImportedMaterialSlotNames.IsValid())
@@ -297,7 +298,7 @@ namespace DatasmithOpenNurbsTranslatorUtils
 			PolygonGroupImportedMaterialSlotNames[PolyGroupId] = DatasmithMeshHelper::DefaultSlotName(0);
 		}
 		// At least one UV set must exist.
-		VertexInstanceUVs.SetNumIndices(1);
+		VertexInstanceUVs.SetNumChannels(1);
 
 		for (int32 MeshIndex = 0; MeshIndex < MeshCount; ++MeshIndex)
 		{
diff --git a/Engine/Plugins/Enterprise/DatasmithCADImporter/Source/DatasmithWireTranslator/Private/OpenModelUtils.cpp b/Engine/Plugins/Enterprise/DatasmithCADImporter/Source/DatasmithWireTranslator/Private/OpenModelUtils.cpp
--- a/Engine/Plugins/Enterprise/DatasmithCADImporter/Source/DatasmithWireTranslator/Private/OpenModelUtils.cpp
+++ b/Engine/Plugins/Enterprise/DatasmithCADImporter/Source/DatasmithWireTranslator/Private/OpenModelUtils.cpp
@@ -1698,9 +1698,9 @@ bool OpenModelUtils::TransferAlMeshToMeshDescription(const AlMesh& AliasMesh, FM
 	PolygonGroupImportedMaterialSlotNames[PolyGroupId] = ImportedSlotName;
 
 	// At least one UV set must exist.
-	if( VertexInstanceUVs.GetNumIndices() == 0 )
+	if( VertexInstanceUVs.GetNumChannels() == 0 )
 	{
-		VertexInstanceUVs.SetNumIndices(1);
+		VertexInstanceUVs.SetNumChannels(1);
 	}
 
 	// Get Alias mesh info
diff --git a/Engine/Plugins/Enterprise/DatasmithFBXImporter/Source/DatasmithDeltaGenTranslator/Private/DatasmithDeltaGenImporter.cpp b/Engine/Plugins/Enterprise/DatasmithFBXImporter/Source/DatasmithDeltaGenTranslator/Private/DatasmithDeltaGenImporter.cpp
--- a/Engine/Plugins/Enterprise/DatasmithFBXImporter/Source/DatasmithDeltaGenTranslator/Private/DatasmithDeltaGenImporter.cpp
+++ b/Engine/Plugins/Enterprise/DatasmithFBXImporter/Source/DatasmithDeltaGenTranslator/Private/DatasmithDeltaGenImporter.cpp
@@ -296,7 +296,7 @@ TSharedPtr<IDatasmithActorElement> FDatasmithDeltaGenImporter::ConvertNode(const
 			FMeshDescription& MeshDescription = ThisMesh->MeshDescription;
 			FStaticMeshAttributes StaticMeshAttributes(MeshDescription);
 			TVertexInstanceAttributesRef<FVector2D> VertexInstanceUVs = StaticMeshAttributes.GetVertexInstanceUVs();
-			int32 NumUVChannels = VertexInstanceUVs.GetNumIndices();
+			int32 NumUVChannels = VertexInstanceUVs.GetNumChannels();
 
 			// DeltaGen uses UV channel 0 for texture UVs, and UV channel 1 for lightmap UVs
 			// Don't set it to zero or else it will disable Datasmith's GenerateLightmapUV option
diff --git a/Engine/Plugins/Enterprise/DatasmithFBXImporter/Source/DatasmithFBXTranslator/Private/DatasmithFBXFileImporter.cpp b/Engine/Plugins/Enterprise/DatasmithFBXImporter/Source/DatasmithFBXTranslator/Private/DatasmithFBXFileImporter.cpp
--- a/Engine/Plugins/Enterprise/DatasmithFBXImporter/Source/DatasmithFBXTranslator/Private/DatasmithFBXFileImporter.cpp
+++ b/Engine/Plugins/Enterprise/DatasmithFBXImporter/Source/DatasmithFBXTranslator/Private/DatasmithFBXFileImporter.cpp
@@ -1030,7 +1030,6 @@ void FDatasmithFBXFileImporter::DoImportMesh(FbxMesh* InMesh, FDatasmithFBXScene
 	TVertexInstanceAttributesRef<FVector4> VertexInstanceColors = MeshDescription.VertexInstanceAttributes().GetAttributesRef<FVector4>(MeshAttribute::VertexInstance::Color);
 	TVertexInstanceAttributesRef<FVector2D> VertexInstanceUVs = MeshDescription.VertexInstanceAttributes().GetAttributesRef<FVector2D>(MeshAttribute::VertexInstance::TextureCoordinate);
 	TEdgeAttributesRef<bool> EdgeHardnesses = MeshDescription.EdgeAttributes().GetAttributesRef<bool>(MeshAttribute::Edge::IsHard);
-	TEdgeAttributesRef<float> EdgeCreaseSharpnesses = MeshDescription.EdgeAttributes().GetAttributesRef<float>(MeshAttribute::Edge::CreaseSharpness);
 	TPolygonGroupAttributesRef<FName> PolygonGroupImportedMaterialSlotNames = MeshDescription.PolygonGroupAttributes().GetAttributesRef<FName>(MeshAttribute::PolygonGroup::ImportedMaterialSlotName);
 
 	// Reserve space for attributes.
@@ -1046,7 +1045,7 @@ void FDatasmithFBXFileImporter::DoImportMesh(FbxMesh* InMesh, FDatasmithFBXScene
 
 	// At least one UV set must exist.
 	int32 MeshDescUVCount = FMath::Max(1, FbxUVCount);
-	VertexInstanceUVs.SetNumIndices(MeshDescUVCount);
+	VertexInstanceUVs.SetNumChannels(MeshDescUVCount);
 
 	//Fill the vertex array
 	for (int32 VertexIndex = 0; VertexIndex < VertexCount; ++VertexIndex)
@@ -1073,7 +1072,7 @@ void FDatasmithFBXFileImporter::DoImportMesh(FbxMesh* InMesh, FDatasmithFBXScene
 	auto GetOrCreatePolygonGroupId = [&](int32 MaterialIndex)
 	{
 		FPolygonGroupID& PolyGroupId = PolygonGroupMapping.FindOrAdd(MaterialIndex);
-		if (PolyGroupId == FPolygonGroupID::Invalid)
+		if (PolyGroupId == INDEX_NONE)
 		{
 			PolyGroupId = MeshDescription.CreatePolygonGroup();
 			FName ImportedSlotName = DatasmithMeshHelper::DefaultSlotName(MaterialIndex);
@@ -1217,13 +1216,12 @@ void FDatasmithFBXFileImporter::DoImportMesh(FbxMesh* InMesh, FDatasmithFBXScene
 
 			// Get or create edge
 			FEdgeID EdgeId = MeshDescription.GetVertexPairEdge(CornerVertexIDs[CornerIndex], CornerVertexIDs[NextCornerIndex]);
-			if (EdgeId == FEdgeID::Invalid)
+			if (EdgeId == INDEX_NONE)
 			{
 				EdgeId = MeshDescription.CreateEdge(CornerVertexIDs[CornerIndex], CornerVertexIDs[NextCornerIndex]);
 
 				// Crease sharpness
 				int32 FbxEdgeIndex = InMesh->GetMeshEdgeIndexForPolygon(PolygonIndex, CornerIndex);
-				EdgeCreaseSharpnesses[EdgeId] = (float)InMesh->GetEdgeCreaseInfo(FbxEdgeIndex);
 
 				// Smoothing
 				bool Hardness = true;
diff --git a/Engine/Plugins/Enterprise/DatasmithIFCImporter/Source/DatasmithIFCTranslator/Private/IFC/IFCStaticMeshFactory.cpp b/Engine/Plugins/Enterprise/DatasmithIFCImporter/Source/DatasmithIFCTranslator/Private/IFC/IFCStaticMeshFactory.cpp
--- a/Engine/Plugins/Enterprise/DatasmithIFCImporter/Source/DatasmithIFCTranslator/Private/IFC/IFCStaticMeshFactory.cpp
+++ b/Engine/Plugins/Enterprise/DatasmithIFCImporter/Source/DatasmithIFCTranslator/Private/IFC/IFCStaticMeshFactory.cpp
@@ -47,7 +47,7 @@ namespace IFC
 		TVertexInstanceAttributesRef<FVector2D> VertexInstanceUVs = MeshDescription->VertexInstanceAttributes().GetAttributesRef<FVector2D>(MeshAttribute::VertexInstance::TextureCoordinate);
 
 		const int32 NumUVs = 1;
-		VertexInstanceUVs.SetNumIndices(NumUVs);
+		VertexInstanceUVs.SetNumChannels(NumUVs);
 
 		FIndexVertexIdMap PositionIndexToVertexId;
 		PositionIndexToVertexId.Empty(InObject->facesVerticesCount);
diff --git a/Engine/Plugins/Enterprise/DatasmithImporter/Source/DatasmithImporter/Private/DatasmithBlueprintLibrary.cpp b/Engine/Plugins/Enterprise/DatasmithImporter/Source/DatasmithImporter/Private/DatasmithBlueprintLibrary.cpp
--- a/Engine/Plugins/Enterprise/DatasmithImporter/Source/DatasmithImporter/Private/DatasmithBlueprintLibrary.cpp
+++ b/Engine/Plugins/Enterprise/DatasmithImporter/Source/DatasmithImporter/Private/DatasmithBlueprintLibrary.cpp
@@ -436,7 +436,7 @@ void UDatasmithStaticMeshBlueprintLibrary::SetupStaticLighting(const TArray< UOb
 					FStaticMeshConstAttributes Attributes(*MeshDescription);
 
 					// 3 is the maximum that lightmass accept. Defined in MeshExport.h : MAX_TEXCOORDS .
-					MaxBiggestUVChannel = FMath::Min(MaxBiggestUVChannel, Attributes.GetVertexInstanceUVs().GetNumIndices() - 1);
+					MaxBiggestUVChannel = FMath::Min(MaxBiggestUVChannel, Attributes.GetVertexInstanceUVs().GetNumChannels() - 1);
 				}
 
 				if (bGenerateLightmapUVs)
diff --git a/Engine/Plugins/Enterprise/DatasmithImporter/Source/DatasmithImporter/Private/DatasmithStaticMeshImporter.cpp b/Engine/Plugins/Enterprise/DatasmithImporter/Source/DatasmithImporter/Private/DatasmithStaticMeshImporter.cpp
--- a/Engine/Plugins/Enterprise/DatasmithImporter/Source/DatasmithImporter/Private/DatasmithStaticMeshImporter.cpp
+++ b/Engine/Plugins/Enterprise/DatasmithImporter/Source/DatasmithImporter/Private/DatasmithStaticMeshImporter.cpp
@@ -25,6 +25,7 @@
 #include "PhysicsEngine/AggregateGeom.h"
 #include "PhysicsEngine/BodySetup.h"
 #include "StaticMeshAttributes.h"
+#include "StaticMeshOperations.h"
 #include "Templates/UniquePtr.h"
 #include "UObject/Package.h"
 #include "UVTools/UVGenerationFlattenMapping.h"
@@ -147,6 +148,12 @@ void FDatasmithStaticMeshImporter::CleanupMeshDescriptions(TArray<FMeshDescripti
 			FElementIDRemappings Remappings;
 			MeshDescription.Compact(Remappings);
 		}
+
+		// Fix invalid vertex normals and tangents
+		// We need polygon info because ComputeTangentsAndNormals uses it to repair the invalid vertex normals/tangents
+		// Can't calculate just the required polygons as ComputeTangentsAndNormals is parallel and we can't guarantee thread-safe access patterns
+		FStaticMeshOperations::ComputeTriangleTangentsAndNormals(MeshDescription);
+		FStaticMeshOperations::ComputeTangentsAndNormals(MeshDescription, EComputeNTBsFlags::UseMikkTSpace);
 	}
 }
 
diff --git a/Engine/Plugins/Enterprise/DatasmithImporter/Source/DatasmithTranslator/Private/Utility/DatasmithMeshHelper.cpp b/Engine/Plugins/Enterprise/DatasmithImporter/Source/DatasmithTranslator/Private/Utility/DatasmithMeshHelper.cpp
--- a/Engine/Plugins/Enterprise/DatasmithImporter/Source/DatasmithTranslator/Private/Utility/DatasmithMeshHelper.cpp
+++ b/Engine/Plugins/Enterprise/DatasmithImporter/Source/DatasmithTranslator/Private/Utility/DatasmithMeshHelper.cpp
@@ -45,21 +45,15 @@ namespace DatasmithMeshHelper
 		FStaticMeshAttributes(Mesh).Register();
 	}
 
-	bool IsTriangleDegenerated(const FMeshDescription& Mesh, const FMeshTriangle& MeshTriangle)
+	bool IsTriangleDegenerated(const FMeshDescription& Mesh, const FTriangleID TriangleID)
 	{
-		// extract triangle vertices positions
-		FVector P[3];
 		const TVertexAttributesConstRef<FVector> VertexPositions = Mesh.VertexAttributes().GetAttributesRef<FVector>(MeshAttribute::Vertex::Position);
-		for (int32 CornerIndex = 0; CornerIndex < 3; ++CornerIndex)
-		{
-			FVertexID VertexID = Mesh.GetVertexInstanceVertex(MeshTriangle.GetVertexInstanceID(CornerIndex));
-			P[CornerIndex] = VertexPositions[VertexID];
-		}
-
-		float NormalLengthSquare = ((P[1] - P[2]) ^ (P[0] - P[2])).SizeSquared();
+		TArrayView<const FVertexID> TriangleVerts = Mesh.GetTriangleVertices(TriangleID);
+		float NormalLengthSquare = ((VertexPositions[TriangleVerts[1]] - VertexPositions[TriangleVerts[2]]) ^ (VertexPositions[TriangleVerts[0]] - VertexPositions[TriangleVerts[2]])).SizeSquared();
 		return NormalLengthSquare < SMALL_NUMBER;
 	}
 
+
 	void RemoveEmptyPolygonGroups(FMeshDescription& Mesh)
 	{
 		bool bRemovedSection = false;
@@ -232,7 +226,7 @@ namespace DatasmithMeshHelper
 	int32 GetNumUVChannel(const FMeshDescription& Mesh)
 	{
 		const TVertexInstanceAttributesConstRef<FVector2D> UVChannels = Mesh.VertexInstanceAttributes().GetAttributesRef<FVector2D>(MeshAttribute::VertexInstance::TextureCoordinate);
-		return UVChannels.GetNumIndices();
+		return UVChannels.GetNumChannels();
 	}
 
 	bool HasUVChannel(const FMeshDescription& Mesh, int32 ChannelIndex)
@@ -267,7 +261,7 @@ namespace DatasmithMeshHelper
 		if (!HasUVChannel(Mesh, ChannelIndex))
 		{
 			TVertexInstanceAttributesRef<FVector2D> UVChannels = Mesh.VertexInstanceAttributes().GetAttributesRef<FVector2D>(MeshAttribute::VertexInstance::TextureCoordinate);
-			UVChannels.SetNumIndices(ChannelIndex + 1);
+			UVChannels.SetNumChannels(ChannelIndex + 1);
 		}
 		return true;
 	}
@@ -327,9 +321,9 @@ namespace DatasmithMeshHelper
 		FStaticMeshOperations::GenerateBoxUV(MeshDescription, UVParameters, TexCoords);
 
 		TMeshAttributesRef<FVertexInstanceID, FVector2D> UVs = MeshDescription.VertexInstanceAttributes().GetAttributesRef<FVector2D>(MeshAttribute::VertexInstance::TextureCoordinate);
-		if (UVs.GetNumIndices() == 0)
+		if (UVs.GetNumChannels() == 0)
 		{
-			UVs.SetNumIndices(1);
+			UVs.SetNumChannels(1);
 		}
 
 		for (const FVertexInstanceID& VertexInstanceID : MeshDescription.VertexInstances().GetElementIDs())
diff --git a/Engine/Plugins/Enterprise/StaticMeshEditorExtension/Source/MeshProcessingLibrary/Private/MeshDescriptionAdapter.cpp b/Engine/Plugins/Enterprise/StaticMeshEditorExtension/Source/MeshProcessingLibrary/Private/MeshDescriptionAdapter.cpp
--- a/Engine/Plugins/Enterprise/StaticMeshEditorExtension/Source/MeshProcessingLibrary/Private/MeshDescriptionAdapter.cpp
+++ b/Engine/Plugins/Enterprise/StaticMeshEditorExtension/Source/MeshProcessingLibrary/Private/MeshDescriptionAdapter.cpp
@@ -12,7 +12,7 @@ namespace MeshDescriptionAdapterUtils
 	// NOTE: This is the same as FMeshDescription::GetPolygonEdges(), although the method should be renamed to GetPolygonPerimeterEdges()
 	void GetPolygonPerimeterEdges(const FMeshDescription& MeshDescription, const FPolygonID& PolygonID, TArray<FEdgeID>& OutPolygonPerimeterEdgeIDs)
 	{
-		const TArray<FVertexInstanceID>& VertexInstanceIDs = MeshDescription.GetPolygonVertexInstances(PolygonID);
+		TArray<FVertexInstanceID> VertexInstanceIDs = MeshDescription.GetPolygonVertexInstances(PolygonID);
 		int32 EdgeCount = VertexInstanceIDs.Num();
 
 		OutPolygonPerimeterEdgeIDs.SetNumUninitialized(EdgeCount, false);
@@ -33,13 +33,15 @@ FMeshDescriptionAdapter::FMeshDescriptionAdapter(FMeshDescription& InMeshDescrip
 	: MeshDescription(InMeshDescription)
 {
 	// Polygon normals are a transient attribute, so register the attribute here if it is not already registered.
+	PRAGMA_DISABLE_DEPRECATION_WARNINGS
 	MeshDescription.PolygonAttributes().RegisterAttribute<FVector>(MeshAttribute::Polygon::Normal, 1, FVector::ZeroVector, EMeshAttributeFlags::Transient);
-
 	TPolygonAttributesRef<FVector> PolygonNormals = MeshDescription.PolygonAttributes().GetAttributesRef<FVector>(MeshAttribute::Polygon::Normal);
+	PRAGMA_ENABLE_DEPRECATION_WARNINGS
+
 	TVertexAttributesConstRef<FVector> VertexPositions = MeshDescription.VertexAttributes().GetAttributesRef<FVector>(MeshAttribute::Vertex::Position);
 	for (const FPolygonID TriangleID : MeshDescription.Polygons().GetElementIDs())
 	{
-		const TArray<FVertexInstanceID>& VertexInstanceIDs = MeshDescription.GetPolygonVertexInstances(TriangleID);
+		TArray<FVertexInstanceID> VertexInstanceIDs = MeshDescription.GetPolygonVertexInstances(TriangleID);
 
 		// @todo: RichTW: This code is assuming polygons with 3 vertices, when a valid mesh description may have polygons with any number of vertices.
 		// Are triangles the expected input here? We should check(VertexInstanceIDs.Num() == 3), or make this code more flexible.
@@ -441,19 +443,19 @@ void FMeshDescriptionAdapter::SetStatusRecursively()
 
 		for (int32 Corner = 0; Corner < 3; ++Corner)
 		{
-			if (PolygonEdges[Corner] != FEdgeID::Invalid)
+			if (PolygonEdges[Corner] != INDEX_NONE)
 			{
 				EdgeMetaData[PolygonEdges[Corner].GetValue()].Markers |= ELEMENT_STATUS_MASK;
 
 				const FVertexID VertexID0 = MeshDescription.GetEdgeVertex(PolygonEdges[Corner], 0);
 				const FVertexID VertexID1 = MeshDescription.GetEdgeVertex(PolygonEdges[Corner], 1);
 
-				if (VertexID0 != FVertexID::Invalid)
+				if (VertexID0 != INDEX_NONE)
 				{
 					VertexMetaData[VertexID0.GetValue()].Markers |= ELEMENT_STATUS_MASK;
 				}
 
-				if (VertexID1 != FVertexID::Invalid)
+				if (VertexID1 != INDEX_NONE)
 				{
 					VertexMetaData[VertexID1.GetValue()].Markers |= ELEMENT_STATUS_MASK;
 				}
@@ -473,19 +475,19 @@ void FMeshDescriptionAdapter::SetStatusRecursively(uint8_t CategoryMask)
 
 		for (int32 Corner = 0; Corner < 3; ++Corner)
 		{
-			if (PolygonEdges[Corner] != FEdgeID::Invalid && IsEdgeOfCategory(PolygonEdges[Corner].GetValue(), CategoryMask))
+			if (PolygonEdges[Corner] != INDEX_NONE && IsEdgeOfCategory(PolygonEdges[Corner].GetValue(), CategoryMask))
 			{
 				EdgeMetaData[PolygonEdges[Corner].GetValue()].Markers |= ELEMENT_STATUS_MASK;
 
 				const FVertexID VertexID0 = MeshDescription.GetEdgeVertex(PolygonEdges[Corner], 0);
 				const FVertexID VertexID1 = MeshDescription.GetEdgeVertex(PolygonEdges[Corner], 1);
 
-				if (VertexID0 != FVertexID::Invalid && IsVertexOfCategory(VertexID0.GetValue(), CategoryMask))
+				if (VertexID0 != INDEX_NONE && IsVertexOfCategory(VertexID0.GetValue(), CategoryMask))
 				{
 					VertexMetaData[VertexID0.GetValue()].Markers |= ELEMENT_STATUS_MASK;
 				}
 
-				if (VertexID1 != FVertexID::Invalid && IsVertexOfCategory(VertexID1.GetValue(), CategoryMask))
+				if (VertexID1 != INDEX_NONE && IsVertexOfCategory(VertexID1.GetValue(), CategoryMask))
 				{
 					VertexMetaData[VertexID1.GetValue()].Markers |= ELEMENT_STATUS_MASK;
 				}
@@ -505,19 +507,19 @@ void FMeshDescriptionAdapter::ResetStatusRecurcively()
 
 		for (int32 Corner = 0; Corner < 3; ++Corner)
 		{
-			if (PolygonEdges[Corner] != FEdgeID::Invalid)
+			if (PolygonEdges[Corner] != INDEX_NONE)
 			{
 				EdgeMetaData[PolygonEdges[Corner].GetValue()].Markers &= ~ELEMENT_STATUS_MASK;
 
 				const FVertexID VertexID0 = MeshDescription.GetEdgeVertex(PolygonEdges[Corner], 0);
 				const FVertexID VertexID1 = MeshDescription.GetEdgeVertex(PolygonEdges[Corner], 1);
 
-				if (VertexID0 != FVertexID::Invalid)
+				if (VertexID0 != INDEX_NONE)
 				{
 					VertexMetaData[VertexID0.GetValue()].Markers &= ~ELEMENT_STATUS_MASK;
 				}
 
-				if (VertexID1 != FVertexID::Invalid)
+				if (VertexID1 != INDEX_NONE)
 				{
 					VertexMetaData[VertexID1.GetValue()].Markers &= ~ELEMENT_STATUS_MASK;
 				}
@@ -539,19 +541,19 @@ void FMeshDescriptionAdapter::ResetMarkersRecursively()
 
 		for (int32 Corner = 0; Corner < 3; ++Corner)
 		{
-			if (PolygonEdges[Corner] != FEdgeID::Invalid)
+			if (PolygonEdges[Corner] != INDEX_NONE)
 			{
 				EdgeMetaData[PolygonEdges[Corner].GetValue()].Markers &= Mask;
 
 				const FVertexID VertexID0 = MeshDescription.GetEdgeVertex(PolygonEdges[Corner], 0);
 				const FVertexID VertexID1 = MeshDescription.GetEdgeVertex(PolygonEdges[Corner], 1);
 
-				if (VertexID0 != FVertexID::Invalid)
+				if (VertexID0 != INDEX_NONE)
 				{
 					VertexMetaData[VertexID0.GetValue()].Markers &= Mask;
 				}
 
-				if (VertexID1 != FVertexID::Invalid)
+				if (VertexID1 != INDEX_NONE)
 				{
 					VertexMetaData[VertexID1.GetValue()].Markers &= Mask;
 				}
@@ -710,7 +712,7 @@ int FMeshDescriptionAdapter::GetTriangleEdges(FMSTriangleID TriangleIndex, FMSEd
 int FMeshDescriptionAdapter::GetEdgesAndVerticesForTriangle(FMSTriangleID TriangleIndex, FMSEdgeID* EdgeSet, FMSVertexID* NodeSet)
 {
 	const FPolygonID TriangleID(TriangleIndex);
-	const TArray<FVertexInstanceID>& VertexInstanceIDs = MeshDescription.GetPolygonVertexInstances(TriangleID);
+	TArray<FVertexInstanceID> VertexInstanceIDs = MeshDescription.GetPolygonVertexInstances(TriangleID);
 
 	TArray<FEdgeID> TriangleEdges;
 	MeshDescription.GetPolygonPerimeterEdges(TriangleID, TriangleEdges);
@@ -728,7 +730,7 @@ int FMeshDescriptionAdapter::GetEdgesAndVerticesForTriangle(FMSTriangleID Triang
 bool FMeshDescriptionAdapter::GetEdgeDirectionForTriangle(FMSTriangleID TriangleIndex, FMSEdgeID EdgeIndex)
 {
 	const FPolygonID TriangleID(TriangleIndex);
-	const TArray<FVertexInstanceID>& VertexInstanceIDs = MeshDescription.GetPolygonVertexInstances(TriangleID);
+	TArray<FVertexInstanceID> VertexInstanceIDs = MeshDescription.GetPolygonVertexInstances(TriangleID);
 
 	TArray<FEdgeID> TriangleEdges;
 	MeshDescription.GetPolygonPerimeterEdges(TriangleID, TriangleEdges);
@@ -755,7 +757,7 @@ bool FMeshDescriptionAdapter::GetEdgeDirectionForTriangle(FMSTriangleID Triangle
 int FMeshDescriptionAdapter::GetEdgeDirectionsForTriangle(FMSTriangleID TriangleIndex, bool *EdgeDirections)
 {
 	const FPolygonID TriangleID(TriangleIndex);
-	const TArray<FVertexInstanceID>& VertexInstanceIDs = MeshDescription.GetPolygonVertexInstances(TriangleID);
+	TArray<FVertexInstanceID> VertexInstanceIDs = MeshDescription.GetPolygonVertexInstances(TriangleID);
 
 	TArray<FEdgeID> TriangleEdges;
 	MeshDescription.GetPolygonPerimeterEdges(TriangleID, TriangleEdges);
@@ -779,7 +781,7 @@ int FMeshDescriptionAdapter::GetEdgeDirectionsForTriangle(FMSTriangleID Triangle
 FMSVertexID FMeshDescriptionAdapter::GetOppositeVertexOnEdgeForTriangle(FMSTriangleID TriangleIndex, FMSEdgeID EdgeIndex)
 {
 	const FPolygonID TriangleID(TriangleIndex);
-	const TArray<FVertexInstanceID>& VertexInstanceIDs = MeshDescription.GetPolygonVertexInstances(TriangleID);
+	TArray<FVertexInstanceID> VertexInstanceIDs = MeshDescription.GetPolygonVertexInstances(TriangleID);
 
 	TArray<FEdgeID> TriangleEdges;
 	MeshDescription.GetPolygonPerimeterEdges(TriangleID, TriangleEdges);
@@ -824,7 +826,9 @@ FMSEdgeID FMeshDescriptionAdapter::GetOppositeEdgeAtVertexForTriangle(FMSTriangl
 int FMeshDescriptionAdapter::GetTriangleNormal(FMSTriangleID TriangleIndex, MeshSimplifier::Vector3D& Normal)
 {
 	const FPolygonID TriangleID(TriangleIndex);
+	PRAGMA_DISABLE_DEPRECATION_WARNINGS
 	const FVector Vector = MeshDescription.PolygonAttributes().GetAttribute<FVector>(TriangleID, MeshAttribute::Polygon::Normal);
+	PRAGMA_ENABLE_DEPRECATION_WARNINGS
 
 	Normal.set(-Vector[0], -Vector[1], -Vector[2]);
 
@@ -834,7 +838,9 @@ int FMeshDescriptionAdapter::GetTriangleNormal(FMSTriangleID TriangleIndex, Mesh
 int FMeshDescriptionAdapter::GetTriangleInvertedNormal(FMSTriangleID TriangleIndex, MeshSimplifier::Vector3D& Normal)
 {
 	const FPolygonID TriangleID(TriangleIndex);
+	PRAGMA_DISABLE_DEPRECATION_WARNINGS
 	const FVector Vector = MeshDescription.PolygonAttributes().GetAttribute<FVector>(TriangleID, MeshAttribute::Polygon::Normal);
+	PRAGMA_ENABLE_DEPRECATION_WARNINGS
 
 	Normal.set(Vector[0], Vector[1], Vector[2]);
 
@@ -844,7 +850,9 @@ int FMeshDescriptionAdapter::GetTriangleInvertedNormal(FMSTriangleID TriangleInd
 int FMeshDescriptionAdapter::GetTriangleNormalizedNormal(FMSTriangleID TriangleIndex, MeshSimplifier::Vector3D& Normal)
 {
 	const FPolygonID TriangleID(TriangleIndex);
+	PRAGMA_DISABLE_DEPRECATION_WARNINGS
 	const FVector Vector = MeshDescription.PolygonAttributes().GetAttribute<FVector>(TriangleID, MeshAttribute::Polygon::Normal).GetSafeNormal();
+	PRAGMA_ENABLE_DEPRECATION_WARNINGS
 
 	Normal.set(Vector[0], Vector[1], Vector[2]);
 
@@ -1122,7 +1130,7 @@ int FMeshDescriptionAdapter::SetEdgePartitionBorder(FMSEdgeID EdgeIndex)
 bool FMeshDescriptionAdapter::IsEdgeAtUVDiscontinuity(FMSEdgeID EdgeIndex)
 {
 	TVertexInstanceAttributesConstRef<FVector2D> MeshTextureCoordinates = MeshDescription.VertexInstanceAttributes().GetAttributesRef<FVector2D>(MeshAttribute::VertexInstance::TextureCoordinate);
-	if (MeshTextureCoordinates.GetNumIndices() == 0)
+	if (MeshTextureCoordinates.GetNumChannels() == 0)
 	{
 		return false;
 	}
@@ -1146,8 +1154,8 @@ bool FMeshDescriptionAdapter::IsEdgeAtUVDiscontinuity(FMSEdgeID EdgeIndex)
 	VertexInstances[1][0] = MeshDescription.GetVertexInstanceForPolygonVertex(EdgeConnectedPolygons[1], EdgeVertexID0);
 	VertexInstances[1][1] = MeshDescription.GetVertexInstanceForPolygonVertex(EdgeConnectedPolygons[1], EdgeVertexID1);
 
-	if (VertexInstances[0][0] != FVertexInstanceID::Invalid && VertexInstances[0][1] != FVertexInstanceID::Invalid &&
-		VertexInstances[1][0] != FVertexInstanceID::Invalid && VertexInstances[1][1] != FVertexInstanceID::Invalid)
+	if (VertexInstances[0][0] != INDEX_NONE && VertexInstances[0][1] != INDEX_NONE &&
+		VertexInstances[1][0] != INDEX_NONE && VertexInstances[1][1] != INDEX_NONE)
 	{
 		FVector2D VertexInstanceTexCoords[2][2];
 
@@ -1276,7 +1284,7 @@ int FMeshDescriptionAdapter::RemoveAsFeatureLine(FMSEdgeID EdgeIndex)
 FMSEdgeID FMeshDescriptionAdapter::FindEdgeFromVertices(FMSVertexID StartVertexIndex, FMSVertexID EndVertexIndex)
 {
 	FEdgeID EdgeID = MeshDescription.GetVertexPairEdge(FVertexID(StartVertexIndex), FVertexID(EndVertexIndex));
-	return EdgeID == FEdgeID::Invalid ? BAD_INDEX : EdgeID.GetValue();
+	return EdgeID == INDEX_NONE ? BAD_INDEX : EdgeID.GetValue();
 }
 
 bool FMeshDescriptionAdapter::GetEdgeDirectionAtFirstTriangle(FMSEdgeID EdgeIndex)
@@ -1350,9 +1358,9 @@ FMSVertexID FMeshDescriptionAdapter::GetSharedVerticesBetweenEdges(FMSEdgeID Edg
 	const FVertexID OtherEdgeVertexID0 = MeshDescription.GetEdgeVertex(OtherEdgeID, 0);
 	const FVertexID OtherEdgeVertexID1 = MeshDescription.GetEdgeVertex(OtherEdgeID, 1);
 
-	FVertexID VertexID = EdgeVertexID0 == OtherEdgeVertexID0 || EdgeVertexID0 == OtherEdgeVertexID1 ? EdgeVertexID0 : ( EdgeVertexID1 == OtherEdgeVertexID0 || EdgeVertexID1 == OtherEdgeVertexID1 ? EdgeVertexID1 : FVertexID::Invalid);
+	FVertexID VertexID = EdgeVertexID0 == OtherEdgeVertexID0 || EdgeVertexID0 == OtherEdgeVertexID1 ? EdgeVertexID0 : ( EdgeVertexID1 == OtherEdgeVertexID0 || EdgeVertexID1 == OtherEdgeVertexID1 ? EdgeVertexID1 : INDEX_NONE);
 
-	return VertexID == FVertexID::Invalid ? BAD_INDEX : VertexID.GetValue();
+	return VertexID == INDEX_NONE ? BAD_INDEX : VertexID.GetValue();
 }
 
 bool FMeshDescriptionAdapter::IsEdgeLimitedByVertex(FMSEdgeID EdgeIndex, FMSVertexID VertexIndex)
@@ -1488,7 +1496,7 @@ MeshSimplifier::Point3D FMeshDescriptionAdapter::GetVertexCoordinates(FMSVertexI
 
 int FMeshDescriptionAdapter::GetConnectingEdgesAtVertex(FMSVertexID VertexIndex, std::vector<FMSEdgeID>& OutEdges)
 {
-	const TArray<FEdgeID>& VertexEdges = MeshDescription.GetVertexConnectedEdges(FVertexID(VertexIndex));
+	TArrayView<const FEdgeID> VertexEdges = MeshDescription.GetVertexConnectedEdges(FVertexID(VertexIndex));
 
 	OutEdges.resize(VertexEdges.Num());
 	for (int32 Index = 0; Index < VertexEdges.Num(); ++Index)
@@ -1506,7 +1514,7 @@ int FMeshDescriptionAdapter::GetEdgesAtVertex(FMSVertexID VertexIndex, std::vect
 		return 0;
 	}
 
-	const TArray<FEdgeID>& VertexEdges = MeshDescription.GetVertexConnectedEdges(FVertexID(VertexIndex));
+	TArrayView<const FEdgeID> VertexEdges = MeshDescription.GetVertexConnectedEdges(FVertexID(VertexIndex));
 
 	Edges.clear();
 
@@ -1525,7 +1533,7 @@ int FMeshDescriptionAdapter::GetEdgesAtVertex(FMSVertexID VertexIndex, std::vect
 
 int FMeshDescriptionAdapter::GetVertexConnectingTriangles(FMSVertexID VertexIndex, std::unordered_set<FMSTriangleID> &OutTriangles)
 {
-	const TArray<FEdgeID>& EdgeIDs = MeshDescription.GetVertexConnectedEdges(FVertexID(VertexIndex));
+	TArrayView<const FEdgeID> EdgeIDs = MeshDescription.GetVertexConnectedEdges(FVertexID(VertexIndex));
 
 	OutTriangles.clear();
 
@@ -1555,7 +1563,7 @@ int FMeshDescriptionAdapter::ValidateVertexTopology(FMSVertexID VertexIndex)
 
 	// Logic took from GPNode::evaluateConnectivity
 	const FVertexID VertexID(VertexIndex);
-	const TArray<FEdgeID>& VertexConnectedEdgeIDs = MeshDescription.GetVertexConnectedEdges(VertexID);
+	TArrayView<const FEdgeID> VertexConnectedEdgeIDs = MeshDescription.GetVertexConnectedEdges(VertexID);
 
 	switch (VertexConnectedEdgeIDs.Num())
 	{
@@ -1592,7 +1600,7 @@ int FMeshDescriptionAdapter::ValidateVertexTopology(FMSVertexID VertexIndex)
 					{
 						// count the number of face between the both border edge
 						FEdgeID EdgeID(VertexConnectedEdgeIDs[EdgeIndex]);
-						FPolygonID TriangleID(FPolygonID::Invalid);
+						FPolygonID TriangleID(INDEX_NONE);
 						int32 TriangleCount = 0;
 
 						do
@@ -1658,7 +1666,7 @@ int FMeshDescriptionAdapter::ValidateVertexTopology(FMSVertexID VertexIndex)
 				const FEdgeID FirstEdgeID(VertexConnectedEdgeIDs[0]);
 				FEdgeID EdgeID(FirstEdgeID);
 
-				FPolygonID TriangleID(FPolygonID::Invalid);
+				FPolygonID TriangleID(INDEX_NONE);
 				int32 TriangleCount = 0;
 
 				do
@@ -1743,7 +1751,10 @@ int FMeshDescriptionAdapter::ValidateMesh()
 	VertexMetaData.SetNumZeroed(MeshDescription.Vertices().GetArraySize());
 
 	TVertexAttributesConstRef<FVector> VertexPositions = MeshDescription.VertexAttributes().GetAttributesRef<FVector>(MeshAttribute::Vertex::Position);
+	PRAGMA_DISABLE_DEPRECATION_WARNINGS
 	TPolygonAttributesConstRef<FVector> PolygonNormals = MeshDescription.PolygonAttributes().GetAttributesRef<FVector>(MeshAttribute::Polygon::Normal);
+	PRAGMA_ENABLE_DEPRECATION_WARNINGS
+
 	for (const FPolygonID TriangleID : MeshDescription.Polygons().GetElementIDs())
 	{
 		TriangleMetaData[TriangleID.GetValue()].Category = (uint16)MeshSimplifier::EElementCategory::ElementCategorySurface;
@@ -1841,19 +1852,19 @@ void FMeshDescriptionAdapter::ValidateTriangles(std::vector<FMSTriangleID>& Tria
 		}
 	}
 
-	FPolygonGroupID PolygonGroupID = FPolygonGroupID::Invalid;
+	FPolygonGroupID PolygonGroupID = INDEX_NONE;
 	for (FMSTriangleID TriangleIndex : OtherTriangles)
 	{
 		// TODO : Find the right one
 		const FPolygonGroupID TrianglePolygonGroupID = MeshDescription.GetPolygonPolygonGroup(FPolygonID(TriangleIndex));
-		if (TrianglePolygonGroupID != FPolygonGroupID::Invalid)
+		if (TrianglePolygonGroupID != INDEX_NONE)
 		{
 			PolygonGroupID = TrianglePolygonGroupID;
 			break;
 		}
 	}
 
-	if (PolygonGroupID != FPolygonGroupID::Invalid)
+	if (PolygonGroupID != INDEX_NONE)
 	{
 		for (FMSTriangleID TriangleIndex : Triangles)
 		{
@@ -1894,7 +1905,7 @@ int FMeshDescriptionAdapter::CreateTriangle(FMSVertexID VertexIndexA, FMSVertexI
 	TVertexInstanceAttributesRef<float> VertexInstanceBinormalSigns = MeshDescription.VertexInstanceAttributes().GetAttributesRef<float>(MeshAttribute::VertexInstance::BinormalSign);
 	TVertexInstanceAttributesRef<FVector4> VertexInstanceColors = MeshDescription.VertexInstanceAttributes().GetAttributesRef<FVector4>(MeshAttribute::VertexInstance::Color);
 	TVertexInstanceAttributesRef<FVector2D> VertexInstanceUVs = MeshDescription.VertexInstanceAttributes().GetAttributesRef<FVector2D>(MeshAttribute::VertexInstance::TextureCoordinate);
-	int32 NumTexCoords = VertexInstanceUVs.GetNumIndices();
+	int32 NumTexCoords = VertexInstanceUVs.GetNumChannels();
 
 	// Try to reuse existing Vertex instances to create the new polygon
 	TArray<FVertexInstanceID> TriangleVertexInstanceIDs;
@@ -1902,9 +1913,9 @@ int FMeshDescriptionAdapter::CreateTriangle(FMSVertexID VertexIndexA, FMSVertexI
 	for (int32 Corner = 0; Corner < 3; ++Corner)
 	{
 		// vertex instance's color, uvs are copied from first instance of vertex if such instance exists
-		const TArray<FVertexInstanceID>& VertexInstances = MeshDescription.GetVertexVertexInstances(VertexIDs[Corner]);
+		TArrayView<const FVertexInstanceID> VertexInstances = MeshDescription.GetVertexVertexInstances(VertexIDs[Corner]);
 
-		FVertexInstanceID VertexInstanceID = FVertexInstanceID::Invalid;
+		FVertexInstanceID VertexInstanceID = INDEX_NONE;
 		for (const FVertexInstanceID CandidateVertexInstance : VertexInstances)
 		{
 			if ((VertexInstanceNormals[CandidateVertexInstance] ^ TriangleNormal).SizeSquared() < 0.00001f)
@@ -1915,7 +1926,7 @@ int FMeshDescriptionAdapter::CreateTriangle(FMSVertexID VertexIndexA, FMSVertexI
 		}
 
 		// No match, let's create a new one
-		if (VertexInstanceID == FVertexInstanceID::Invalid)
+		if (VertexInstanceID == INDEX_NONE)
 		{
 			VertexInstanceID = MeshDescription.CreateVertexInstance(VertexIDs[Corner]);
 
@@ -1945,7 +1956,7 @@ int FMeshDescriptionAdapter::CreateTriangle(FMSVertexID VertexIndexA, FMSVertexI
 		EdgeVertexIDs[1] = VertexIDs[CornerIndices[1]];
 
 		FEdgeID MatchEdgeId = MeshDescription.GetVertexPairEdge(EdgeVertexIDs[0], EdgeVertexIDs[1]);
-		if (MatchEdgeId == FEdgeID::Invalid)
+		if (MatchEdgeId == INDEX_NONE)
 		{
 			MatchEdgeId = MeshDescription.CreateEdge(EdgeVertexIDs[0], EdgeVertexIDs[1]);
 			int32 Increment = MeshDescription.Edges().GetArraySize() - EdgeMetaData.Num();
@@ -1964,7 +1975,9 @@ int FMeshDescriptionAdapter::CreateTriangle(FMSVertexID VertexIndexA, FMSVertexI
 	}
 	const FPolygonID NewPolygonID = MeshDescription.CreatePolygon(PolygonGroupID, TriangleVertexInstanceIDs);
 
+	PRAGMA_DISABLE_DEPRECATION_WARNINGS
 	MeshDescription.PolygonAttributes().SetAttribute(NewPolygonID, MeshAttribute::Polygon::Normal, 0, TriangleNormal);
+	PRAGMA_ENABLE_DEPRECATION_WARNINGS
 
 	{
 		int32 Increment = MeshDescription.Polygons().GetArraySize() - TriangleMetaData.Num();
diff --git a/Engine/Plugins/Enterprise/StaticMeshEditorExtension/Source/MeshProcessingLibrary/Private/MeshProcessingLibrary.cpp b/Engine/Plugins/Enterprise/StaticMeshEditorExtension/Source/MeshProcessingLibrary/Private/MeshProcessingLibrary.cpp
--- a/Engine/Plugins/Enterprise/StaticMeshEditorExtension/Source/MeshProcessingLibrary/Private/MeshProcessingLibrary.cpp
+++ b/Engine/Plugins/Enterprise/StaticMeshEditorExtension/Source/MeshProcessingLibrary/Private/MeshProcessingLibrary.cpp
@@ -656,10 +656,10 @@ void UMeshProcessingLibrary::ApplyJacketingOnMeshActors(const TArray<AActor*>& A
 			RawMesh->Empty();
 
 			//Create the missing normals and tangents on polygons because FStaticMeshOperations::ComputeTangentsAndNormals requires it
-			if(!NewRawMesh.PolygonAttributes().GetAttributesRef<FVector>(MeshAttribute::Polygon::Normal).IsValid()
-				|| !NewRawMesh.PolygonAttributes().GetAttributesRef<FVector>(MeshAttribute::Polygon::Tangent).IsValid())
+			if(!NewRawMesh.PolygonAttributes().GetAttributesRef<FVector>(MeshAttribute::Triangle::Normal).IsValid()
+				|| !NewRawMesh.PolygonAttributes().GetAttributesRef<FVector>(MeshAttribute::Triangle::Tangent).IsValid())
 			{
-				FStaticMeshOperations::ComputePolygonTangentsAndNormals(NewRawMesh);
+				FStaticMeshOperations::ComputeTriangleTangentsAndNormals(NewRawMesh);
 			}
 
 			const FMeshBuildSettings& BuildSettings = StaticMesh->GetSourceModel(0).BuildSettings;
diff --git a/Engine/Plugins/Enterprise/StaticMeshEditorExtension/Source/StaticMeshEditorExtension/Private/MeshEditorUtils.cpp b/Engine/Plugins/Enterprise/StaticMeshEditorExtension/Source/StaticMeshEditorExtension/Private/MeshEditorUtils.cpp
--- a/Engine/Plugins/Enterprise/StaticMeshEditorExtension/Source/StaticMeshEditorExtension/Private/MeshEditorUtils.cpp
+++ b/Engine/Plugins/Enterprise/StaticMeshEditorExtension/Source/StaticMeshEditorExtension/Private/MeshEditorUtils.cpp
@@ -202,12 +202,14 @@ TMap<FPolygonID, FTriangleData> MeshEditingUtilsImpl::GetFilteredTriangleData(co
 	FilteredTriangleData.Reset();
 
 	TVertexAttributesConstRef<FVector> VertexPositions = MeshDescription->VertexAttributes().GetAttributesRef<FVector>(MeshAttribute::Vertex::Position);
+	PRAGMA_DISABLE_DEPRECATION_WARNINGS
 	TPolygonAttributesConstRef<FVector> PolygonCenters = MeshDescription->PolygonAttributes().GetAttributesRef<FVector>(MeshAttribute::Polygon::Center);
+	PRAGMA_ENABLE_DEPRECATION_WARNINGS
 
 	FVector TriangleVertexPositions[3];
 	for (const FPolygonID PolygonID : Polygons)
 	{
-		const TArray<FTriangleID>& PolygonTriangleIDs = MeshDescription->GetPolygonTriangleIDs(PolygonID);
+		TArrayView<const FTriangleID> PolygonTriangleIDs = MeshDescription->GetPolygonTriangleIDs(PolygonID);
 
 		for (const FTriangleID MeshTriangleID : PolygonTriangleIDs)
 		{
diff --git a/Engine/Plugins/Enterprise/StaticMeshEditorExtension/Source/StaticMeshEditorExtension/Private/UVTools/SUVGenerationTool.cpp b/Engine/Plugins/Enterprise/StaticMeshEditorExtension/Source/StaticMeshEditorExtension/Private/UVTools/SUVGenerationTool.cpp
--- a/Engine/Plugins/Enterprise/StaticMeshEditorExtension/Source/StaticMeshEditorExtension/Private/UVTools/SUVGenerationTool.cpp
+++ b/Engine/Plugins/Enterprise/StaticMeshEditorExtension/Source/StaticMeshEditorExtension/Private/UVTools/SUVGenerationTool.cpp
@@ -258,7 +258,7 @@ FReply SGenerateUV::OnApplyUV() const
 		const int32 CurrentLOD = GetSelectedLOD(EditorPtr);
 		const FMeshDescription* MeshDescription = StaticMesh->GetMeshDescription(CurrentLOD);
 		FStaticMeshConstAttributes Attributes(*MeshDescription);
-		const int32 NumberOfUVChannels = Attributes.GetVertexInstanceUVs().GetNumIndices();
+		const int32 NumberOfUVChannels = Attributes.GetVertexInstanceUVs().GetNumChannels();
 		int32 TargetChannel = FMath::Min((int32)GenerateUVSettings->TargetChannel, EditorPtr->GetNumUVChannels());
 		const int32 NumberOfChannelsToAdd = TargetChannel - (NumberOfUVChannels - 1);
 
diff --git a/Engine/Plugins/Enterprise/StaticMeshEditorExtension/Source/StaticMeshEditorExtension/Private/UVTools/UVGenerationFlattenMapping.cpp b/Engine/Plugins/Enterprise/StaticMeshEditorExtension/Source/StaticMeshEditorExtension/Private/UVTools/UVGenerationFlattenMapping.cpp
--- a/Engine/Plugins/Enterprise/StaticMeshEditorExtension/Source/StaticMeshEditorExtension/Private/UVTools/UVGenerationFlattenMapping.cpp
+++ b/Engine/Plugins/Enterprise/StaticMeshEditorExtension/Source/StaticMeshEditorExtension/Private/UVTools/UVGenerationFlattenMapping.cpp
@@ -1122,7 +1122,7 @@ TArray<FUVGenerationFlattenMappingInternal::FaceStruct> FUVGenerationFlattenMapp
 
 	for (const FPolygonID PolygonID : InMesh.Polygons().GetElementIDs())
 	{
-		const TArray<FTriangleID>& TriangleIDs = InMesh.GetPolygonTriangleIDs(PolygonID);
+		TArrayView<const FTriangleID> TriangleIDs = InMesh.GetPolygonTriangleIDs(PolygonID);
 		for (const FTriangleID TriangleID : TriangleIDs)
 		{
 			FaceStruct Face;
@@ -1218,9 +1218,9 @@ void UUVGenerationFlattenMapping::GenerateUVs(FMeshDescription& InMesh, int32 UV
 
 	const TArray<int32> OverlappingCorners = GetOverlappingCornersRemapping(InMesh, bRemoveDegenerates);
 	const TVertexInstanceAttributesRef<FVector2D> UVChannels = InMesh.VertexInstanceAttributes().GetAttributesRef<FVector2D>(MeshAttribute::VertexInstance::TextureCoordinate);
-	if (UVChannels.GetNumIndices() < UVChannel + 1)
+	if (UVChannels.GetNumChannels() < UVChannel + 1)
 	{
-		UVChannels.SetNumIndices(UVChannel + 1);
+		UVChannels.SetNumChannels(UVChannel + 1);
 	}
 
 	// Get the Internal Face Array
@@ -1260,7 +1260,7 @@ void UUVGenerationFlattenMapping::GenerateUVs(FMeshDescription& InMesh, int32 UV
 				VertexInstanceAttributesSet.ForEach(
 					[&](const FName AttributeName, auto AttributeArrayRef)
 					{
-						for(int32 Index = 0; Index < AttributeArrayRef.GetNumIndices(); ++Index)
+						for(int32 Index = 0; Index < AttributeArrayRef.GetNumChannels(); ++Index)
 						{
 							AttributeArrayRef.Set( NewInstanceID, Index, AttributeArrayRef.Get(OriginalID, Index) );
 						}
@@ -1272,8 +1272,11 @@ void UUVGenerationFlattenMapping::GenerateUVs(FMeshDescription& InMesh, int32 UV
 
 				// Update contour of associated polygon with new vertex instance
 				const FPolygonID PolygonID( Face.PolygonID );
-				const int32 FoundIndex = InMesh.GetPolygonVertexInstances( PolygonID ).Find( OriginalID );
-				InMesh.SetPolygonVertexInstance( PolygonID, FoundIndex, NewInstanceID );
+				TArray<FVertexInstanceID, TInlineAllocator<8>> PolygonVerts = InMesh.GetPolygonVertexInstances<TInlineAllocator<8>>( PolygonID );
+				const int32 FoundIndex = PolygonVerts.Find( OriginalID );
+				check(FoundIndex != INDEX_NONE);
+				PolygonVerts[FoundIndex] = NewInstanceID;
+				InMesh.SetPolygonVertexInstances( PolygonID, PolygonVerts );
 
 				// Mark polygon for re-triangulation
 				PolygonsToRetriangulate.Add( Face.PolygonID );
diff --git a/Engine/Plugins/Enterprise/StaticMeshEditorExtension/StaticMeshEditorExtension.uplugin b/Engine/Plugins/Enterprise/StaticMeshEditorExtension/StaticMeshEditorExtension.uplugin
--- a/Engine/Plugins/Enterprise/StaticMeshEditorExtension/StaticMeshEditorExtension.uplugin
+++ b/Engine/Plugins/Enterprise/StaticMeshEditorExtension/StaticMeshEditorExtension.uplugin
@@ -13,7 +13,8 @@
 	"EnabledByDefault" : false,
 	"CanContainContent" : true,
 	"IsBetaVersion" : true,
-	"Installed" : false,
+	"Installed": false,
+	"Hidden":  true,
 	"Modules" :
 	[
 		{
diff --git a/Engine/Plugins/Experimental/AlembicImporter/Source/AlembicLibrary/Private/AbcImporter.cpp b/Engine/Plugins/Experimental/AlembicImporter/Source/AlembicLibrary/Private/AbcImporter.cpp
--- a/Engine/Plugins/Experimental/AlembicImporter/Source/AlembicLibrary/Private/AbcImporter.cpp
+++ b/Engine/Plugins/Experimental/AlembicImporter/Source/AlembicLibrary/Private/AbcImporter.cpp
@@ -1227,7 +1227,6 @@ void FAbcImporter::GenerateMeshDescriptionFromSample(const FAbcMeshSample* Sampl
 
 	TVertexAttributesRef<FVector> VertexPositions = Attributes.GetVertexPositions();
 	TEdgeAttributesRef<bool> EdgeHardnesses = Attributes.GetEdgeHardnesses();
-	TEdgeAttributesRef<float> EdgeCreaseSharpnesses = Attributes.GetEdgeCreaseSharpnesses();
 	TPolygonGroupAttributesRef<FName> PolygonGroupImportedMaterialSlotNames = Attributes.GetPolygonGroupMaterialSlotNames();
 	TVertexInstanceAttributesRef<FVector> VertexInstanceNormals = Attributes.GetVertexInstanceNormals();
 	TVertexInstanceAttributesRef<FVector> VertexInstanceTangents = Attributes.GetVertexInstanceTangents();
@@ -1236,7 +1235,7 @@ void FAbcImporter::GenerateMeshDescriptionFromSample(const FAbcMeshSample* Sampl
 	TVertexInstanceAttributesRef<FVector2D> VertexInstanceUVs = Attributes.GetVertexInstanceUVs();
 
 	//Speedtree use UVs to store is data
-	VertexInstanceUVs.SetNumIndices(Sample->NumUVSets);
+	VertexInstanceUVs.SetNumChannels(Sample->NumUVSets);
 	
 	for (int32 MatIndex = 0; MatIndex < StaticMesh->StaticMaterials.Num(); ++MatIndex)
 	{
diff --git a/Engine/Plugins/Experimental/BlastPlugin/Source/BlastAuthoring/Private/MeshUtility.cpp b/Engine/Plugins/Experimental/BlastPlugin/Source/BlastAuthoring/Private/MeshUtility.cpp
--- a/Engine/Plugins/Experimental/BlastPlugin/Source/BlastAuthoring/Private/MeshUtility.cpp
+++ b/Engine/Plugins/Experimental/BlastPlugin/Source/BlastAuthoring/Private/MeshUtility.cpp
@@ -53,7 +53,7 @@ void FMeshUtility::EditableMeshToBlastMesh(const UEditableMesh* SourceMesh, Nv::
 
 	for(const FPolygonID PolygonID : MeshDescription->Polygons().GetElementIDs())
 	{
-		const TArray<FTriangleID>& TriangleIDs = MeshDescription->GetPolygonTriangleIDs(PolygonID);
+		TArrayView<const FTriangleID> TriangleIDs = MeshDescription->GetPolygonTriangleIDs(PolygonID);
 		for(const FTriangleID TriangleID : TriangleIDs)
 		{
 			for(int32 TriVertIndex = 0; TriVertIndex < 3; ++TriVertIndex)
@@ -120,7 +120,7 @@ void FMeshUtility::EditableMeshToBlastMesh(const UEditableMesh* SourceMesh, int3
 		BlastUVs.Emplace(UV.X, UV.Y);
 	}
 
-	const TArray<FPolygonID>& PolygonGroupIDs = MeshDescription->GetPolygonGroupPolygons(FPolygonGroupID(PolygonGroup));
+	TArrayView<const FPolygonID> PolygonGroupIDs = MeshDescription->GetPolygonGroupPolygons(FPolygonGroupID(PolygonGroup));
 
 	TSharedPtr<FGeometryCollection, ESPMode::ThreadSafe> GeometryCollectionSPtr = GeometryCollection->GetGeometryCollection();
 	FGeometryCollection* GeometryCollectionPtr = GeometryCollectionSPtr.Get();
@@ -146,7 +146,7 @@ void FMeshUtility::EditableMeshToBlastMesh(const UEditableMesh* SourceMesh, int3
 
 	for(const FPolygonID PolygonID : PolygonGroupIDs)
 	{
-		const TArray<FTriangleID>& TriangleIDs = MeshDescription->GetPolygonTriangleIDs(PolygonID);
+		TArrayView<const FTriangleID> TriangleIDs = MeshDescription->GetPolygonTriangleIDs(PolygonID);
 
 		check(TriangleIDs.Num() == 1);
 
