Adding missed build.cs file from CL 14084099

Move FEditorModeInfo and shared mode defines to EditorFramework module.

#rnx

diff --git a/Engine/Plugins/Enterprise/StaticMeshEditorExtension/Source/StaticMeshEditorExtension/StaticMeshEditorExtension.Build.cs b/Engine/Plugins/Enterprise/StaticMeshEditorExtension/Source/StaticMeshEditorExtension/StaticMeshEditorExtension.Build.cs
--- a/Engine/Plugins/Enterprise/StaticMeshEditorExtension/Source/StaticMeshEditorExtension/StaticMeshEditorExtension.Build.cs
+++ b/Engine/Plugins/Enterprise/StaticMeshEditorExtension/Source/StaticMeshEditorExtension/StaticMeshEditorExtension.Build.cs
@@ -28,6 +28,7 @@ namespace UnrealBuildTool.Rules
 					"SlateCore",
 					"StaticMeshDescription",
 					"StaticMeshEditor",
+					"EditorFramework",
 					"UnrealEd",
 					"ViewportInteraction",
 					"DataprepCore",
