Fixed crash with the UV generation tool because the plugin is being made obsolete

diff --git a/Engine/Plugins/Enterprise/StaticMeshEditorExtension/Source/StaticMeshEditorExtension/Private/UVTools/SUVGenerationTool.cpp b/Engine/Plugins/Enterprise/StaticMeshEditorExtension/Source/StaticMeshEditorExtension/Private/UVTools/SUVGenerationTool.cpp
--- a/Engine/Plugins/Enterprise/StaticMeshEditorExtension/Source/StaticMeshEditorExtension/Private/UVTools/SUVGenerationTool.cpp
+++ b/Engine/Plugins/Enterprise/StaticMeshEditorExtension/Source/StaticMeshEditorExtension/Private/UVTools/SUVGenerationTool.cpp
@@ -153,6 +153,11 @@ void SGenerateUV::Construct(const FArguments& InArgs)
 
 int32 GetSelectedLOD(const TSharedPtr<IStaticMeshEditor>& EditorPtr)
 {
+	if ( !EditorPtr || !EditorPtr->GetStaticMeshComponent() || !EditorPtr->GetStaticMesh() )
+	{
+		return 0;
+	}
+
 	const int32 SelectedLOD = EditorPtr->GetStaticMeshComponent()->ForcedLodModel - 1;
 
 	//If there is only one LOD then there is no ambiguity, even in AutoLOD.
